Introduced the screen reader and Slate screen reader plugins for Windows. Provides basic screen reader
functionality for Slate where focused widgets are announced to the user when the screen reader is active.
The screen reader plugin sets up the screen reading framework and the Slate screen reader plugin provides concrete implementations of provided interfaces and base classes.

diff --git a/Engine/Plugins/Experimental/ScreenReader/ScreenReader.uplugin b/Engine/Plugins/Experimental/ScreenReader/ScreenReader.uplugin
new file mode 100644
--- /dev/null
+++ b/Engine/Plugins/Experimental/ScreenReader/ScreenReader.uplugin
@@ -0,0 +1,40 @@
+{
+	"FileVersion" : 3,
+	"Version" : 1,
+	"VersionName" : "0.1",
+	"FriendlyName" : "ScreenReader",
+	"Description" : "A plugin that contains accessibility classes and frameworks that can be extended to offer vision accessibility services.",
+	"Category" : "Accessibility",
+	"CreatedBy" : "Epic Games, Inc.",
+	"CreatedByURL" : "http://epicgames.com",
+	"DocsURL" : "",
+	"MarketplaceURL" : "",
+	"SupportURL" : "",
+	"EnabledByDefault" : false,
+	"CanContainContent" : false,
+	"IsExperimental" : true,
+	"Installed" : false,
+	"Modules" :
+	[
+		{
+			"Name" : "ScreenReader",
+			"Type" : "Runtime",
+			"LoadingPhase" : "PostEngineInit",
+			"WhitelistPlatforms" :
+			[
+				"Win64"
+			],
+			"BlacklistTargets" :
+			[
+				"Server"
+			]
+		}
+	],
+	"Plugins" :
+	[
+		{
+			"Name" : "TextToSpeech",
+			"Enabled": true
+		}
+	]
+}
diff --git a/Engine/Plugins/Experimental/ScreenReader/Source/ScreenReader/Private/Announcement/ScreenReaderAnnouncement.cpp b/Engine/Plugins/Experimental/ScreenReader/Source/ScreenReader/Private/Announcement/ScreenReaderAnnouncement.cpp
new file mode 100644
--- /dev/null
+++ b/Engine/Plugins/Experimental/ScreenReader/Source/ScreenReader/Private/Announcement/ScreenReaderAnnouncement.cpp
@@ -0,0 +1,62 @@
+// Copyright Epic Games, Inc. All Rights Reserved.
+
+#include "Announcement/ScreenReaderAnnouncement.h"
+#include "Misc/StringBuilder.h"
+#include "ScreenReaderLog.h"
+
+FScreenReaderAnnouncementInfo::FScreenReaderAnnouncementInfo(bool bInShouldQueue, bool bInInterruptable, EScreenReaderAnnouncementPriority InPriority)
+	: bShouldQueue( bInShouldQueue)
+	, bInterruptable(bInInterruptable)
+	, Priority(InPriority)
+	, Timestamp(FPlatformTime::Seconds())
+{
+
+}
+
+FString FScreenReaderAnnouncementInfo::ToString() const
+{
+	TStringBuilder<128> Builder;
+	Builder.Append(TEXT("Announcement info: Should queue = "));
+	const FString ShouldQueueValue = ShouldQueue() ? "true" : "false";
+	Builder.Append(ShouldQueueValue);
+	Builder.Append(TEXT(". Is interruptable = "));
+	const FString IsInterruptableValue = IsInterruptable() ? "true" : "false";
+	Builder.Append(IsInterruptableValue);
+	Builder.Append(TEXT(". Priority = "));
+	switch (GetPriority())
+	{
+	case EScreenReaderAnnouncementPriority::High:
+		Builder.Append(TEXT("High."));
+		break;
+	case EScreenReaderAnnouncementPriority::Medium:
+		Builder.Append(TEXT("Medium"));
+		break;
+	case EScreenReaderAnnouncementPriority::Low:
+		Builder.Append(TEXT("Low"));
+		break;
+	default:
+		UE_LOG(LogScreenReaderAnnouncement, Error, TEXT("Unknown priorityannouncement priority level."));
+		break;
+	}
+	Builder.Append(TEXT(". Timestamp = "));
+
+
+	return Builder.ToString();
+}
+
+FScreenReaderAnnouncement::FScreenReaderAnnouncement(FString InAnnouncementString, FScreenReaderAnnouncementInfo InAnnouncementInfo)
+	: AnnouncementString(MoveTemp(InAnnouncementString))
+	, AnnouncementInfo(MoveTemp(InAnnouncementInfo))
+{
+
+}
+
+FString FScreenReaderAnnouncement::ToString() const
+{
+	TStringBuilder<256> Builder;
+	Builder.Append(TEXT("AnnouncementString = "));
+	Builder.Append(GetAnnouncementString());
+	Builder.Append(TEXT(". \n"));
+	Builder.Append(GetAnnouncementInfo().ToString());
+	return Builder.ToString();
+}
\ No newline at end of file
diff --git a/Engine/Plugins/Experimental/ScreenReader/Source/ScreenReader/Private/Announcement/ScreenReaderAnnouncementChannel.cpp b/Engine/Plugins/Experimental/ScreenReader/Source/ScreenReader/Private/Announcement/ScreenReaderAnnouncementChannel.cpp
new file mode 100644
--- /dev/null
+++ b/Engine/Plugins/Experimental/ScreenReader/Source/ScreenReader/Private/Announcement/ScreenReaderAnnouncementChannel.cpp
@@ -0,0 +1,114 @@
+// Copyright Epic Games, Inc. All Rights Reserved.
+
+#include "Announcement/ScreenReaderAnnouncementChannel.h"
+#include "TextToSpeech.h"
+#include "ScreenReaderLog.h"
+#include "Announcement/ScreenReaderWidgetAnnouncementBuilder.h"
+
+namespace ScreenReaderAnnouncementChannel_Private
+{
+	static FScreenReaderWidgetAnnouncementBuilder WidgetAnnouncementBuilder;
+}
+
+FScreenReaderAnnouncementChannel::FScreenReaderAnnouncementChannel(const TSharedRef<FTextToSpeechBase>& InTextToSpeech)
+	: TextToSpeech(InTextToSpeech)
+{
+	InTextToSpeech->SetTextToSpeechFinishedSpeakingDelegate(FTextToSpeechBase::FOnTextToSpeechFinishSpeaking::CreateRaw(this, &FScreenReaderAnnouncementChannel::OnTextToSpeechFinishSpeaking));
+}
+
+FScreenReaderAnnouncementChannel::~FScreenReaderAnnouncementChannel()
+{
+	Deactivate();
+}
+
+void FScreenReaderAnnouncementChannel::Activate()
+{
+	ClearAllAnnouncements();
+	TextToSpeech->Activate();
+}
+
+void FScreenReaderAnnouncementChannel::Deactivate()
+{
+	TextToSpeech->Deactivate();
+	ClearAllAnnouncements();
+}
+
+FScreenReaderReply FScreenReaderAnnouncementChannel::RequestSpeak(FScreenReaderAnnouncement InAnnouncement)
+{
+	if (InAnnouncement.GetAnnouncementString().IsEmpty() || !TextToSpeech->IsActive())
+	{
+		return FScreenReaderReply::Unhandled();
+	}
+	
+	UE_LOG(LogScreenReaderAnnouncement, VeryVerbose, TEXT("Requesting to speak announcement %s."), *InAnnouncement.ToString());
+	if (IsSpeaking())
+	{
+		ensureMsgf(!CurrentAnnouncement.GetAnnouncementString().IsEmpty(), TEXT("Channel is Current announcement should have a valid string if announcement is being spoken. Check this algorithm."));
+		UE_LOG(LogScreenReaderAnnouncement, VeryVerbose, TEXT("Announcement currently being spoken."));
+		// If the current announcement is currently being spoken,
+		// We see if it can be interrupted by the passed in message 
+		if (FScreenReaderAnnouncement::CanBeInterruptedBy(CurrentAnnouncement, InAnnouncement, FScreenReaderAnnouncement::FDefaultInterruptionPolicy()))
+		{
+			if (CurrentAnnouncement.GetAnnouncementInfo().ShouldQueue())
+			{
+				AnnouncementQueue.HeapPush(MoveTemp(CurrentAnnouncement), FScreenReaderAnnouncement::FDefaultComparator());
+			}
+			CurrentAnnouncement = MoveTemp(InAnnouncement);
+			TextToSpeech->Speak(CurrentAnnouncement.GetAnnouncementString());
+		}
+		else if (InAnnouncement.GetAnnouncementInfo().ShouldQueue())
+		{
+			AnnouncementQueue.HeapPush(MoveTemp(InAnnouncement), FScreenReaderAnnouncement::FDefaultComparator());
+		}
+	}
+	else
+	{
+		ensureMsgf(CurrentAnnouncement.GetAnnouncementString().IsEmpty(), TEXT("No announcement is being spoken, but current announcement has a valid string. Check your algorithm."));
+		UE_LOG(LogScreenReaderAnnouncement, VeryVerbose, TEXT("Nothing being spoken. Speaking immediately."));
+		// nothing beimng spoken, speak passed in announcement immediately 
+		CurrentAnnouncement = MoveTemp(InAnnouncement);
+		TextToSpeech->Speak(CurrentAnnouncement.GetAnnouncementString());
+	}
+	return FScreenReaderReply::Handled();
+}
+
+FScreenReaderReply FScreenReaderAnnouncementChannel::StopSpeaking()
+{
+	TextToSpeech->StopSpeaking();
+	return FScreenReaderReply::Handled();
+}
+
+bool FScreenReaderAnnouncementChannel::IsSpeaking() const
+{
+	return TextToSpeech->IsSpeaking();
+}
+
+FScreenReaderReply FScreenReaderAnnouncementChannel::ClearAllAnnouncements()
+{
+	CurrentAnnouncement = FScreenReaderAnnouncement();
+	AnnouncementQueue.Empty();
+	ensure(CurrentAnnouncement.GetAnnouncementString().IsEmpty());
+	ensure(AnnouncementQueue.Num() == 0);
+	return FScreenReaderReply::Handled();
+}
+
+FScreenReaderReply FScreenReaderAnnouncementChannel::RequestSpeakWidget(const TSharedRef<IAccessibleWidget>& InWidget)
+{
+	FString WidgetAnnouncement = ScreenReaderAnnouncementChannel_Private::WidgetAnnouncementBuilder.BuildWidgetAnnouncement(InWidget);
+	return RequestSpeak(FScreenReaderAnnouncement(WidgetAnnouncement, FScreenReaderAnnouncementInfo::DefaultWidgetAnnouncement()));
+}
+
+void FScreenReaderAnnouncementChannel::OnTextToSpeechFinishSpeaking()
+{
+	// reset the current announcement to null. This is necessary for the RequestSpeak() algorithm 
+	CurrentAnnouncement = FScreenReaderAnnouncement();
+	// if there are queued announcements, speak the next one 
+	if (AnnouncementQueue.Num() > 0)
+	{
+		UE_LOG(LogScreenReaderAnnouncement, VeryVerbose, TEXT("Announcing next queued announcement."));
+		// @TODOAccessibility: We create a temp to ensure the RequestSpeak)() algorithm continues to work. Optimize in future 
+		FScreenReaderAnnouncement NextAnnouncement;
+		AnnouncementQueue.HeapPop(NextAnnouncement, FScreenReaderAnnouncement::FDefaultComparator());
+		RequestSpeak(MoveTemp(NextAnnouncement));
+	 }
+}
\ No newline at end of file
diff --git a/Engine/Plugins/Experimental/ScreenReader/Source/ScreenReader/Private/Announcement/ScreenReaderWidgetAnnouncementBuilder.cpp b/Engine/Plugins/Experimental/ScreenReader/Source/ScreenReader/Private/Announcement/ScreenReaderWidgetAnnouncementBuilder.cpp
new file mode 100644
--- /dev/null
+++ b/Engine/Plugins/Experimental/ScreenReader/Source/ScreenReader/Private/Announcement/ScreenReaderWidgetAnnouncementBuilder.cpp
@@ -0,0 +1,118 @@
+// Copyright Epic Games, Inc. All Rights Reserved.
+
+#include "Announcement/ScreenReaderWidgetAnnouncementBuilder.h"
+#include "GenericPlatform/Accessibility/GenericAccessibleInterfaces.h"
+#include "Misc/StringBuilder.h"
+
+#define LOCTEXT_NAMESPACE "ScreenReaderWidgetAnnouncementBuilkder"
+const TMap<EAccessibleWidgetType, FText> FScreenReaderWidgetAnnouncementBuilder::WidgetTypeToAccessibleRoleMap
+{ 
+	{EAccessibleWidgetType::Button, LOCTEXT("ControlTypeButton", "button")},
+	{EAccessibleWidgetType::CheckBox, LOCTEXT("ControlTypeCheckBox", "check box")},
+	{EAccessibleWidgetType::ComboBox, LOCTEXT("ControlTypeComboBox", "combo box")},
+	{EAccessibleWidgetType::Hyperlink, LOCTEXT("ControlTypeHyperlink", "Link")},
+	{EAccessibleWidgetType::Image, LOCTEXT("ControlTypeImage", "image")},
+	{EAccessibleWidgetType::Layout, LOCTEXT("ControlTypeLayout", "Layout")},
+	{EAccessibleWidgetType::ScrollBar, LOCTEXT("ControlTypeScrollBar", "scroll bar")},
+	{EAccessibleWidgetType::Slider, LOCTEXT("ControlTypeSlider", "slider")},
+	{EAccessibleWidgetType::Text, LOCTEXT("ControlTypeText", "text")},
+	{EAccessibleWidgetType::TextEdit, LOCTEXT("ControlTypeTextEdit", "edit")},
+{EAccessibleWidgetType::Window, LOCTEXT("ControlTypeWindow", "window")},
+{EAccessibleWidgetType::List, LOCTEXT("ControlTypeList", "List")},
+	{EAccessibleWidgetType::ListItem, LOCTEXT("ControlTypeListItem", "ListItem")},
+{EAccessibleWidgetType::ListItem, LOCTEXT("ControlTypeListItem", "ListItem")}
+};
+
+// @TODOAccessibility: These are just sample descriptions. Should go into settings or soemthing 
+const TMap<EAccessibleWidgetType, FText> FScreenReaderWidgetAnnouncementBuilder::WidgetTypeToInteractionDescriptionMap
+{
+	{EAccessibleWidgetType::Button, LOCTEXT("InteractionTypeButton", "Press enter on the keyboard or the confirm button of the gamepad to activate.")},
+	{EAccessibleWidgetType::CheckBox, LOCTEXT("InteractionTypeCheckBox", "Press enter on the keyboard or the confirm button of the gamepad to toggle.")},
+	{EAccessibleWidgetType::ComboBox, LOCTEXT("InteractionTypeComboBox", "Press enter on the keyboard or the confirm button of the gamepad to focus. Repeat to expand.")},
+	{EAccessibleWidgetType::Slider, LOCTEXT("InteractionTypeSlider", "Press enter on the keyboard or the confirm button of the gamepad to focus. Then press left and right to change the value.")},
+	{EAccessibleWidgetType::Slider, LOCTEXT("InteractionTypeSlider", "Press enter on the keyboard or the confirm button of the gamepad to focus. Then press left and right to change the value.")},
+	{EAccessibleWidgetType::List, LOCTEXT("InteractionTypeList", "Press up and down to navigate through the list items.")}
+};
+
+FString FScreenReaderWidgetAnnouncementBuilder::BuildWidgetAnnouncement(const TSharedRef<IAccessibleWidget>& InWidget)
+{
+	// format is label, value, role. Help text. Interaction desription.
+	// This follows formatting of standard OS screen readers 
+	TStringBuilder<256> Announcement;
+	Announcement.Append(BuildLabel(InWidget));
+	Announcement.Append(TEXT(", "));
+	Announcement.Append(BuildValue(InWidget));
+	Announcement.Append(TEXT(", "));	
+	Announcement.Append(BuildRole(InWidget));
+	Announcement.Append(TEXT(". "));
+	Announcement.Append(BuildHelpText(InWidget));
+	Announcement.Append(TEXT(". "));
+	Announcement.Append(BuildInteractionDescription(InWidget));
+	Announcement.Append(TEXT("."));
+
+	return Announcement.ToString();
+}
+
+FString FScreenReaderWidgetAnnouncementBuilder::BuildLabel(const TSharedRef<IAccessibleWidget>& InWidget) const
+{
+	return InWidget->GetWidgetName();
+}
+
+FString FScreenReaderWidgetAnnouncementBuilder::BuildRole(const TSharedRef<IAccessibleWidget>& InWidget) const
+{
+	static const FText UnknownAccessibleRole = LOCTEXT("ControlTypeUnknown", "Unknown");
+	const FText* LocalizedAccessibleRole = FScreenReaderWidgetAnnouncementBuilder::WidgetTypeToAccessibleRoleMap.Find(InWidget->GetWidgetType());
+	return LocalizedAccessibleRole ? LocalizedAccessibleRole->ToString() : UnknownAccessibleRole .ToString();
+}
+
+FString FScreenReaderWidgetAnnouncementBuilder::BuildValue(const TSharedRef<IAccessibleWidget>& InWidget) const
+{
+	if (!InWidget->AsProperty())
+	{
+		return FString();
+	}
+	switch (InWidget->GetWidgetType())
+	{
+	case EAccessibleWidgetType::CheckBox:
+	{
+		if (InWidget->AsActivatable()->GetCheckedState())
+		{
+			static const FText CheckedState= LOCTEXT("CheckBoxChecked", "Checked");
+			return CheckedState.ToString();
+		}
+		else
+		{
+			static const FText UncheckedState = LOCTEXT("CheckBoxUnchecked", "Unchecked");
+			return UncheckedState.ToString();
+		}
+	}
+	case EAccessibleWidgetType::TextEdit:
+	{
+		// only return a value if the field is not a password
+		if (!InWidget->AsProperty()->IsPassword())
+		{
+			return InWidget->AsProperty()->GetValue();
+		}
+		else
+		{
+			static const FText PasswordHidden= LOCTEXT("TextFieldPasswordHidden", "Hidden");
+			return PasswordHidden.ToString();
+		}
+	}
+	default:
+		return InWidget->AsProperty()->GetValue();
+	}
+}
+
+FString FScreenReaderWidgetAnnouncementBuilder::BuildHelpText(const TSharedRef<IAccessibleWidget>& InWidget) const
+{
+	return InWidget->GetHelpText();
+}
+FString FScreenReaderWidgetAnnouncementBuilder::BuildInteractionDescription(const TSharedRef<IAccessibleWidget>& InWidget) const
+{
+	const FText* LocalizedInteractionDescription = FScreenReaderWidgetAnnouncementBuilder::WidgetTypeToInteractionDescriptionMap.Find(InWidget->GetWidgetType());
+	return LocalizedInteractionDescription ? LocalizedInteractionDescription->ToString() : FString();
+}
+
+
+#undef LOCTEXT_NAMESPACE
\ No newline at end of file
diff --git a/Engine/Plugins/Experimental/ScreenReader/Source/ScreenReader/Private/Announcement/ScreenReaderWidgetAnnouncementBuilder.h b/Engine/Plugins/Experimental/ScreenReader/Source/ScreenReader/Private/Announcement/ScreenReaderWidgetAnnouncementBuilder.h
new file mode 100644
--- /dev/null
+++ b/Engine/Plugins/Experimental/ScreenReader/Source/ScreenReader/Private/Announcement/ScreenReaderWidgetAnnouncementBuilder.h
@@ -0,0 +1,27 @@
+// Copyright Epic Games, Inc. All Rights Reserved.
+
+#pragma once
+#include "CoreMinimal.h"
+
+class IAccessibleWidget;
+enum class EAccessibleWidgetType : uint8;
+
+
+class FScreenReaderWidgetAnnouncementBuilder
+{
+public:
+	FScreenReaderWidgetAnnouncementBuilder() = default;
+	~FScreenReaderWidgetAnnouncementBuilder() = default;
+	FString BuildWidgetAnnouncement(const TSharedRef<IAccessibleWidget>& InWidget);
+
+private:
+	FString BuildLabel(const TSharedRef<IAccessibleWidget>& InWidget) const;
+	FString BuildRole(const TSharedRef<IAccessibleWidget>& InWidget) const;
+	FString BuildValue(const TSharedRef<IAccessibleWidget>& InWidget) const;
+	FString BuildHelpText(const TSharedRef<IAccessibleWidget>& InWidget) const;
+	FString BuildInteractionDescription(const TSharedRef<IAccessibleWidget>& InWidget) const;
+
+	static const TMap<EAccessibleWidgetType, FText> WidgetTypeToAccessibleRoleMap;
+	static const TMap<EAccessibleWidgetType, FText> WidgetTypeToInteractionDescriptionMap;
+};
+
diff --git a/Engine/Plugins/Experimental/ScreenReader/Source/ScreenReader/Private/GenericPlatform/ScreenReaderApplicationMessageHandlerBase.cpp b/Engine/Plugins/Experimental/ScreenReader/Source/ScreenReader/Private/GenericPlatform/ScreenReaderApplicationMessageHandlerBase.cpp
new file mode 100644
--- /dev/null
+++ b/Engine/Plugins/Experimental/ScreenReader/Source/ScreenReader/Private/GenericPlatform/ScreenReaderApplicationMessageHandlerBase.cpp
@@ -0,0 +1,255 @@
+// Copyright Epic Games, Inc. All Rights Reserved.
+
+#include "GenericPlatform/ScreenReaderApplicationMessageHandlerBase.h"
+#include "GenericPlatform/ScreenReaderBase.h"
+
+FScreenReaderApplicationMessageHandlerBase::FScreenReaderApplicationMessageHandlerBase(const TSharedRef<FGenericApplicationMessageHandler>& InTargetMessageHandler, FScreenReaderBase& InOwningScreenReader)
+	: OwningScreenReader(InOwningScreenReader)
+	, TargetMessageHandler(InTargetMessageHandler)
+	, bProcessMouseInput(false)
+	, bProcessKeyInput(true)
+{
+
+}
+
+// FGenericApplicationMessageHandler
+bool FScreenReaderApplicationMessageHandlerBase::ShouldProcessUserInputMessages(const TSharedPtr< FGenericWindow >& PlatformWindow) const
+{
+	return TargetMessageHandler->ShouldProcessUserInputMessages(PlatformWindow);
+}
+
+bool FScreenReaderApplicationMessageHandlerBase::OnKeyChar(const TCHAR Character, const bool IsRepeat)
+{
+	return TargetMessageHandler->OnKeyChar(Character, IsRepeat);
+}
+
+bool FScreenReaderApplicationMessageHandlerBase::OnKeyDown(const int32 KeyCode, const uint32 CharacterCode, const bool IsRepeat)
+{
+	return TargetMessageHandler->OnKeyDown(KeyCode, CharacterCode, IsRepeat);
+}
+
+bool FScreenReaderApplicationMessageHandlerBase::OnKeyUp(const int32 KeyCode, const uint32 CharacterCode, const bool IsRepeat)
+{
+	return TargetMessageHandler->OnKeyUp(KeyCode, CharacterCode, IsRepeat);
+}
+
+bool FScreenReaderApplicationMessageHandlerBase::OnMouseDown(const TSharedPtr< FGenericWindow >& Window, const EMouseButtons::Type Button)
+{
+	return TargetMessageHandler->OnMouseDown(Window, Button);
+}
+
+bool FScreenReaderApplicationMessageHandlerBase::OnMouseDown(const TSharedPtr< FGenericWindow >& Window, const EMouseButtons::Type Button, const FVector2D CursorPos)
+{
+	return TargetMessageHandler->OnMouseDown(Window, Button, CursorPos);
+}
+
+bool FScreenReaderApplicationMessageHandlerBase::OnMouseUp(const EMouseButtons::Type Button)
+{
+	return TargetMessageHandler->OnMouseUp(Button);
+}
+
+bool FScreenReaderApplicationMessageHandlerBase::OnMouseUp(const EMouseButtons::Type Button, const FVector2D CursorPos)
+{
+	return TargetMessageHandler->OnMouseUp(Button, CursorPos);
+}
+
+bool FScreenReaderApplicationMessageHandlerBase::OnMouseDoubleClick(const TSharedPtr< FGenericWindow >& Window, const EMouseButtons::Type Button)
+{
+	return TargetMessageHandler->OnMouseDoubleClick(Window, Button);
+}
+
+bool FScreenReaderApplicationMessageHandlerBase::OnMouseDoubleClick(const TSharedPtr< FGenericWindow >& Window, const EMouseButtons::Type Button, const FVector2D CursorPos)
+{
+	return TargetMessageHandler->OnMouseDoubleClick(Window, Button, CursorPos);
+}
+
+bool FScreenReaderApplicationMessageHandlerBase::OnMouseWheel(const float Delta)
+{
+	return TargetMessageHandler->OnMouseWheel(Delta);
+}
+
+bool FScreenReaderApplicationMessageHandlerBase::OnMouseWheel(const float Delta, const FVector2D CursorPos)
+{
+	return TargetMessageHandler->OnMouseWheel(Delta, CursorPos);
+}
+
+bool FScreenReaderApplicationMessageHandlerBase::OnMouseMove()
+{
+	return TargetMessageHandler->OnMouseMove();
+}
+
+bool FScreenReaderApplicationMessageHandlerBase::OnRawMouseMove(const int32 X, const int32 Y)
+{
+	return TargetMessageHandler->OnRawMouseMove(X, Y);
+}
+
+bool FScreenReaderApplicationMessageHandlerBase::OnCursorSet()
+{
+	return TargetMessageHandler->OnCursorSet();
+}
+
+bool FScreenReaderApplicationMessageHandlerBase::OnControllerAnalog(FGamepadKeyNames::Type KeyName, int32 ControllerId, float AnalogValue)
+{
+	return TargetMessageHandler->OnControllerAnalog(KeyName, ControllerId, AnalogValue);
+}
+
+bool FScreenReaderApplicationMessageHandlerBase::OnControllerButtonPressed(FGamepadKeyNames::Type KeyName, int32 ControllerId, bool IsRepeat)
+{
+	return TargetMessageHandler->OnControllerButtonPressed(KeyName, ControllerId, IsRepeat);
+}
+
+bool FScreenReaderApplicationMessageHandlerBase::OnControllerButtonReleased(FGamepadKeyNames::Type KeyName, int32 ControllerId, bool IsRepeat)
+{
+	return TargetMessageHandler->OnControllerButtonReleased(KeyName, ControllerId, IsRepeat);
+}
+
+void FScreenReaderApplicationMessageHandlerBase::OnBeginGesture()
+{
+	TargetMessageHandler->OnBeginGesture();
+}
+
+bool FScreenReaderApplicationMessageHandlerBase::OnTouchGesture(EGestureEvent GestureType, const FVector2D& Delta, float WheelDelta, bool bIsDirectionInvertedFromDevice)
+{
+	return TargetMessageHandler->OnTouchGesture(GestureType, Delta, WheelDelta, bIsDirectionInvertedFromDevice);
+}
+
+void FScreenReaderApplicationMessageHandlerBase::OnEndGesture()
+{
+	TargetMessageHandler->OnEndGesture();
+}
+
+bool FScreenReaderApplicationMessageHandlerBase::OnTouchStarted(const TSharedPtr< FGenericWindow >& Window, const FVector2D& Location, float Force, int32 TouchIndex, int32 ControllerId)
+{
+	return TargetMessageHandler->OnTouchStarted(Window, Location, Force, TouchIndex, ControllerId);
+}
+
+bool FScreenReaderApplicationMessageHandlerBase::OnTouchMoved(const FVector2D& Location, float Force, int32 TouchIndex, int32 ControllerId)
+{
+	return TargetMessageHandler->OnTouchMoved(Location, Force, TouchIndex, ControllerId);
+}
+
+bool FScreenReaderApplicationMessageHandlerBase::OnTouchEnded(const FVector2D& Location, int32 TouchIndex, int32 ControllerId)
+{
+	return TargetMessageHandler->OnTouchEnded(Location, TouchIndex, ControllerId);
+}
+
+bool FScreenReaderApplicationMessageHandlerBase::OnTouchForceChanged(const FVector2D& Location, float Force, int32 TouchIndex, int32 ControllerId)
+{
+	return TargetMessageHandler->OnTouchForceChanged(Location, Force, TouchIndex, ControllerId);
+}
+
+bool FScreenReaderApplicationMessageHandlerBase::OnTouchFirstMove(const FVector2D& Location, float Force, int32 TouchIndex, int32 ControllerId)
+{
+	return TargetMessageHandler->OnTouchFirstMove(Location, Force, TouchIndex, ControllerId);
+}
+
+void FScreenReaderApplicationMessageHandlerBase::ShouldSimulateGesture(EGestureEvent Gesture, bool bEnable)
+{
+	return TargetMessageHandler->ShouldSimulateGesture(Gesture, bEnable);
+}
+
+bool FScreenReaderApplicationMessageHandlerBase::OnMotionDetected(const FVector& Tilt, const FVector& RotationRate, const FVector& Gravity, const FVector& Acceleration, int32 ControllerId)
+{
+	return TargetMessageHandler->OnMotionDetected(Tilt, RotationRate, Gravity, Acceleration, ControllerId);
+}
+
+bool FScreenReaderApplicationMessageHandlerBase::OnSizeChanged(const TSharedRef< FGenericWindow >& Window, const int32 Width, const int32 Height, bool bWasMinimized /* = false */)
+{
+	return TargetMessageHandler->OnSizeChanged(Window, Width, Height, bWasMinimized);
+}
+
+void FScreenReaderApplicationMessageHandlerBase::OnOSPaint(const TSharedRef<FGenericWindow>& Window)
+{
+	return TargetMessageHandler->OnOSPaint(Window);
+}
+
+FWindowSizeLimits FScreenReaderApplicationMessageHandlerBase::GetSizeLimitsForWindow(const TSharedRef<FGenericWindow>& Window) const
+{
+	return TargetMessageHandler->GetSizeLimitsForWindow(Window);
+}
+
+void FScreenReaderApplicationMessageHandlerBase::OnResizingWindow(const TSharedRef< FGenericWindow >& Window)
+{
+	TargetMessageHandler->OnResizingWindow(Window);
+}
+
+bool FScreenReaderApplicationMessageHandlerBase::BeginReshapingWindow(const TSharedRef< FGenericWindow >& Window)
+{
+	return TargetMessageHandler->BeginReshapingWindow(Window);
+}
+
+void FScreenReaderApplicationMessageHandlerBase::FinishedReshapingWindow(const TSharedRef< FGenericWindow >& Window)
+{
+	return TargetMessageHandler->FinishedReshapingWindow(Window);
+}
+
+void FScreenReaderApplicationMessageHandlerBase::HandleDPIScaleChanged(const TSharedRef< FGenericWindow >& Window)
+{
+	TargetMessageHandler->HandleDPIScaleChanged(Window);
+}
+
+void FScreenReaderApplicationMessageHandlerBase::OnMovedWindow(const TSharedRef< FGenericWindow >& Window, const int32 X, const int32 Y)
+{
+	TargetMessageHandler->OnMovedWindow(Window, X, Y);
+}
+
+bool FScreenReaderApplicationMessageHandlerBase::OnWindowActivationChanged(const TSharedRef< FGenericWindow >& Window, const EWindowActivation ActivationType)
+{
+	return TargetMessageHandler->OnWindowActivationChanged(Window, ActivationType);
+}
+
+bool FScreenReaderApplicationMessageHandlerBase::OnApplicationActivationChanged(const bool IsActive)
+{
+	return TargetMessageHandler->OnApplicationActivationChanged(IsActive);
+}
+
+bool FScreenReaderApplicationMessageHandlerBase::OnConvertibleLaptopModeChanged()
+{
+	return TargetMessageHandler->OnConvertibleLaptopModeChanged();
+}
+
+EWindowZone::Type FScreenReaderApplicationMessageHandlerBase::GetWindowZoneForPoint(const TSharedRef< FGenericWindow >& Window, const int32 X, const int32 Y)
+{
+	return TargetMessageHandler->GetWindowZoneForPoint(Window, X, Y);
+}
+
+void FScreenReaderApplicationMessageHandlerBase::OnWindowClose(const TSharedRef< FGenericWindow >& Window)
+{
+	TargetMessageHandler->OnWindowClose(Window);
+}
+
+EDropEffect::Type FScreenReaderApplicationMessageHandlerBase::OnDragEnterText(const TSharedRef< FGenericWindow >& Window, const FString& Text)
+{
+	return TargetMessageHandler->OnDragEnterText(Window, Text);
+}
+
+EDropEffect::Type FScreenReaderApplicationMessageHandlerBase::OnDragEnterFiles(const TSharedRef< FGenericWindow >& Window, const TArray< FString >& Files)
+{
+	return TargetMessageHandler->OnDragEnterFiles(Window, Files);
+}
+
+EDropEffect::Type FScreenReaderApplicationMessageHandlerBase::OnDragEnterExternal(const TSharedRef< FGenericWindow >& Window, const FString& Text, const TArray< FString >& Files)
+{
+	return TargetMessageHandler->OnDragEnterExternal(Window, Text, Files);
+}
+
+EDropEffect::Type FScreenReaderApplicationMessageHandlerBase::OnDragOver(const TSharedPtr< FGenericWindow >& Window)
+{
+	return TargetMessageHandler->OnDragOver(Window);
+}
+
+void FScreenReaderApplicationMessageHandlerBase::OnDragLeave(const TSharedPtr< FGenericWindow >& Window)
+{
+	TargetMessageHandler->OnDragLeave(Window);
+}
+
+EDropEffect::Type FScreenReaderApplicationMessageHandlerBase::OnDragDrop(const TSharedPtr< FGenericWindow >& Window)
+{
+	return TargetMessageHandler->OnDragDrop(Window);
+}
+
+bool FScreenReaderApplicationMessageHandlerBase::OnWindowAction(const TSharedRef< FGenericWindow >& Window, const EWindowAction::Type InActionType)
+{
+	return TargetMessageHandler->OnWindowAction(Window, InActionType);
+}
+// ~
\ No newline at end of file
diff --git a/Engine/Plugins/Experimental/ScreenReader/Source/ScreenReader/Private/GenericPlatform/ScreenReaderBase.cpp b/Engine/Plugins/Experimental/ScreenReader/Source/ScreenReader/Private/GenericPlatform/ScreenReaderBase.cpp
new file mode 100644
--- /dev/null
+++ b/Engine/Plugins/Experimental/ScreenReader/Source/ScreenReader/Private/GenericPlatform/ScreenReaderBase.cpp
@@ -0,0 +1,101 @@
+// Copyright Epic Games, Inc. All Rights Reserved.
+
+#include "GenericPlatform/ScreenReaderBase.h"
+#include "GenericPlatform/ScreenReaderApplicationMessageHandlerBase.h"
+#include "GenericPlatform/GenericApplication.h"
+#include "GenericPlatform/ScreenReaderUser.h"
+#include "Announcement/ScreenReaderAnnouncementChannel.h"
+#include "ScreenReaderLog.h"
+
+
+FScreenReaderBase::FScreenReaderBase(const TSharedRef<GenericApplication>& InPlatformApplication)
+	: ScreenReaderApplicationMessageHandler(MakeShared<FScreenReaderApplicationMessageHandlerBase>(InPlatformApplication->GetMessageHandler(), *this))
+	, PlatformApplication(InPlatformApplication)
+	, bActive(false)
+{
+
+}
+
+FScreenReaderBase::~FScreenReaderBase()
+{
+
+}
+
+void FScreenReaderBase::Activate()
+{
+	if (!IsActive())
+	{
+		checkf(PlatformApplication.IsValid(), TEXT("Trying to activate the screen reader with an invalid platform aplication.with "));
+		UE_LOG(LogScreenReader, Verbose, TEXT("Activating screen reader."));
+		// Set the screen reader application message handler to be the platform's message handler 
+		// This allows the screen reader to intercept all applications before passing them on to Slate 
+		PlatformApplication.Pin()->SetMessageHandler(ScreenReaderApplicationMessageHandler);
+		OnActivate();
+		bActive = true;
+}
+}
+
+void FScreenReaderBase::Deactivate()
+{
+	if (IsActive())
+	{
+		UE_LOG(LogScreenReader, Verbose, TEXT("Deactivating screen reader."));
+		OnDeactivate();
+		// Set FSlateApplication as the platform message ahndler again 
+		PlatformApplication.Pin()->SetMessageHandler(ScreenReaderApplicationMessageHandler->GetTargetMessageHandler());
+		bActive = false;
+	}
+}
+
+bool FScreenReaderBase::RegisterUser(int32 InUserId)
+{
+	if (UsersMap.Contains(InUserId))
+	{
+		UE_LOG(LogScreenReader, Verbose, TEXT("Failed to register screen reader user with Id %i. Another user with the same Id has already been registered."), InUserId);
+		return false;
+	}
+	UsersMap.Add(InUserId, MakeShared<FScreenReaderUser>(InUserId));
+	UE_LOG(LogScreenReader, Verbose, TEXT("Registered screen reader user %i."), InUserId);
+	return true;
+}
+
+bool FScreenReaderBase::UnregisterUser(int32 InUserId)
+{
+	if (!UsersMap.Contains(InUserId))
+	{
+		UE_LOG(LogScreenReader, Verbose, TEXT("Failed to unregister with Id %i. No user with %i Id is registered."), InUserId);
+		return false;
+	}
+	UsersMap[InUserId]->Deactivate();
+	UsersMap.Remove(InUserId);
+	UE_LOG(LogScreenReader, Verbose, TEXT("Unregistered screen reader user %i."), InUserId);
+	return true;
+}
+
+bool FScreenReaderBase::IsUserRegistered(int32 InUserId) const
+{
+	return UsersMap.Contains(InUserId);
+}
+
+void FScreenReaderBase::UnregisterAllUsers()
+{
+	for (TPair<int32, TSharedRef<FScreenReaderUser>>& UserPair : UsersMap)
+	{
+		UserPair.Value->Deactivate();
+	}
+	UsersMap.Empty();
+}
+
+TSharedRef<FScreenReaderUser> FScreenReaderBase::GetUser(int32 InUserId) const
+{
+	return UsersMap[InUserId];
+}
+
+TSharedPtr<FScreenReaderUser> FScreenReaderBase::GetUserChecked(int32 InUserId) const
+{
+	if (IsUserRegistered(InUserId))
+	{
+		return UsersMap[InUserId];
+	}
+	return nullptr;
+}
\ No newline at end of file
diff --git a/Engine/Plugins/Experimental/ScreenReader/Source/ScreenReader/Private/GenericPlatform/ScreenReaderUser.cpp b/Engine/Plugins/Experimental/ScreenReader/Source/ScreenReader/Private/GenericPlatform/ScreenReaderUser.cpp
new file mode 100644
--- /dev/null
+++ b/Engine/Plugins/Experimental/ScreenReader/Source/ScreenReader/Private/GenericPlatform/ScreenReaderUser.cpp
@@ -0,0 +1,72 @@
+// Copyright Epic Games, Inc. All Rights Reserved.
+
+#include "GenericPlatform/ScreenReaderUser.h"
+#include "Announcement/ScreenReaderAnnouncementChannel.h"
+#include "TextToSpeech.h"
+
+FScreenReaderUser::FScreenReaderUser(int32 InUserId)
+	: UserId(InUserId)
+	, bActive(false)
+{
+	// @TODOAccessibility: For now, just default it to the platform TTS. Should give a way to allow users to use custom TTS though 
+	AnnouncementChannel = MakeUnique<FScreenReaderAnnouncementChannel>(ITextToSpeechModule::Get().GetPlatformFactory()->Create());
+}
+
+FScreenReaderUser::~FScreenReaderUser()
+{
+	Deactivate();
+}
+
+void FScreenReaderUser::Activate()
+{
+	if (!IsActive())
+	{
+		bActive = true;
+		AnnouncementChannel->Activate();
+	}
+}
+
+void FScreenReaderUser::Deactivate()
+{
+	if (IsActive())
+	{
+		bActive = false;
+		AnnouncementChannel->Deactivate();
+	}
+}
+
+FScreenReaderReply FScreenReaderUser::RequestSpeak(FScreenReaderAnnouncement InAnnouncement)
+{
+	if (IsActive())
+	{
+		return AnnouncementChannel->RequestSpeak(MoveTemp(InAnnouncement));
+	}
+	return FScreenReaderReply::Unhandled();
+}
+
+FScreenReaderReply FScreenReaderUser::StopSpeaking()
+{
+	if (IsActive())
+	{
+		return AnnouncementChannel->StopSpeaking();
+	}
+	return FScreenReaderReply::Unhandled();
+}
+
+bool FScreenReaderUser::IsSpeaking() const
+{
+	if (IsActive())
+	{
+		return AnnouncementChannel->IsSpeaking();
+	}
+	return false;
+}
+
+FScreenReaderReply FScreenReaderUser::RequestSpeakWidget(const TSharedRef<IAccessibleWidget>& InWidget)
+{
+	if (IsActive())
+	{
+		return AnnouncementChannel->RequestSpeakWidget(InWidget);
+	}
+	return FScreenReaderReply::Unhandled();
+}
\ No newline at end of file
diff --git a/Engine/Plugins/Experimental/ScreenReader/Source/ScreenReader/Private/ScreenReaderLog.h b/Engine/Plugins/Experimental/ScreenReader/Source/ScreenReader/Private/ScreenReaderLog.h
new file mode 100644
--- /dev/null
+++ b/Engine/Plugins/Experimental/ScreenReader/Source/ScreenReader/Private/ScreenReaderLog.h
@@ -0,0 +1,12 @@
+// Copyright Epic Games, Inc. All Rights Reserved.
+
+#pragma once
+#include "CoreMinimal.h"
+
+DECLARE_LOG_CATEGORY_EXTERN(LogScreenReader, Log, All);
+
+DECLARE_LOG_CATEGORY_EXTERN(LogScreenReaderInput, Log, All);
+DECLARE_LOG_CATEGORY_EXTERN(LogScreenReaderAnnouncement, Log, All);
+DECLARE_LOG_CATEGORY_EXTERN(LogScreenReaderNavigation, Log, All);
+DECLARE_LOG_CATEGORY_EXTERN(LogScreenReaderAction, Log, All);
+
diff --git a/Engine/Plugins/Experimental/ScreenReader/Source/ScreenReader/Private/ScreenReaderModule.cpp b/Engine/Plugins/Experimental/ScreenReader/Source/ScreenReader/Private/ScreenReaderModule.cpp
new file mode 100644
--- /dev/null
+++ b/Engine/Plugins/Experimental/ScreenReader/Source/ScreenReader/Private/ScreenReaderModule.cpp
@@ -0,0 +1,16 @@
+// Copyright Epic Games, Inc. All Rights Reserved.
+
+
+#include "ScreenReaderLog.h"
+#include "Modules/ModuleManager.h"
+
+static_assert(WITH_ACCESSIBILITY, TEXT("Trying to use the screen reader plugin with accessibility disabled. Accessibility must be enabled to use this plugin. Either enable accessibility or disable the plugin."));
+
+IMPLEMENT_MODULE(FDefaultModuleImpl, ScreenReader);
+
+DEFINE_LOG_CATEGORY(LogScreenReader);
+DEFINE_LOG_CATEGORY(LogScreenReaderInput);
+DEFINE_LOG_CATEGORY(LogScreenReaderAnnouncement);
+DEFINE_LOG_CATEGORY(LogScreenReaderNavigation);
+DEFINE_LOG_CATEGORY(LogScreenReaderAction);
+
diff --git a/Engine/Plugins/Experimental/ScreenReader/Source/ScreenReader/Public/Announcement/ScreenReaderAnnouncement.h b/Engine/Plugins/Experimental/ScreenReader/Source/ScreenReader/Public/Announcement/ScreenReaderAnnouncement.h
new file mode 100644
--- /dev/null
+++ b/Engine/Plugins/Experimental/ScreenReader/Source/ScreenReader/Public/Announcement/ScreenReaderAnnouncement.h
@@ -0,0 +1,171 @@
+// Copyright Epic Games, Inc. All Rights Reserved.
+
+#pragma once
+#include "CoreMinimal.h"
+#include "ScreenReaderAnnouncement.generated.h"
+
+/**
+* The priority level associated with a screen reader announcement.
+* The priority of an announcement could be used to determine if an announcement A could interrupt announcement B
+* or to help order the announcements if multiple are queued.
+*/
+UENUM()
+enum class EScreenReaderAnnouncementPriority : uint8
+{
+	High = 0,
+	Medium,
+	Low
+};
+
+/**
+* A struct that contains information about how a screen reader announcement should behave when the announcement is requested to be spoken.
+* This struct holds information to determine an announcement's priority in relation to other announcements and determines if an 
+* announcement can be interrupted by other announcements and if an announcement should be queued if it cannot be made immendiately.
+* For C++ users, please use the static methods provided to create instances of FScreenReaderInfo with pre-defined defaults that will
+* satisfy 80% of your use cases.
+* Please see FScreenReaderUser::RequestSpeak() for examples of how to use this struct.
+* @see FScreenReaderUser, FScreenReaderAnnouncement
+*/
+USTRUCT(BlueprintType)
+struct SCREENREADER_API FScreenReaderAnnouncementInfo
+{
+	GENERATED_BODY()
+public:
+	// Only used to satisfy USTRUCT constraints 
+	// Please use ctor that takes parameters 
+	FScreenReaderAnnouncementInfo() = default;
+	FScreenReaderAnnouncementInfo(bool bInShouldQueue, bool bInInterruptable, EScreenReaderAnnouncementPriority InPriority);
+	bool ShouldQueue() const { return bShouldQueue; }
+	bool IsInterruptable() const { return bInterruptable; }
+	EScreenReaderAnnouncementPriority GetPriority() const { return Priority; }
+	double GetTimestamp() const { return Timestamp; }
+	/** Returns a string that represents all of the data reflected in this struct. */
+	FString ToString() const;
+
+	/** 
+	* The announcement info used for announcements concerning widgets and their content.
+	* These are the lowest priority, are interruptable and will not be queued if interrupted.
+	* Use this when requestin any announcements related to widgets that the user can easily access again
+	* E.g Useful for announcing widgets that the user is currently focused on. The user can always refocus on the widget to hear its contents spoken so
+	* it doesn't matter if it is interrupted.
+	*/
+	static FScreenReaderAnnouncementInfo DefaultWidgetAnnouncement()
+	{
+		return FScreenReaderAnnouncementInfo(false, true, EScreenReaderAnnouncementPriority::Low);
+	}
+	/**
+	* An announcement info used for announcements that are of critical importance to the end user.
+	* Announcements with this announcement info have the highest priority, are not interruptable by other announcements and will be queued if multiple important announcements are requested.
+	* Use this for communicating critical information to an end user.
+	* E.g Use this for announcements concerning the system (battery low, accessibility enabled etc) or important in game events that cannot be revisted or played again.
+	*/
+	static FScreenReaderAnnouncementInfo Important()
+	{
+		return FScreenReaderAnnouncementInfo(true, false, EScreenReaderAnnouncementPriority::High);
+	}
+
+	/**
+	* This announcement info is used for announcements to provide feedback for end user actions or inform them of changes.
+	* Announcements with this announcement info have medium priority, can be interrupted and will not be queued if interrupted by an announcement of equal or higher priority.
+	* User this announcement info for announcements that provide feedback to the end user.
+	* E.g When they have successfully activated a UI element, the screen has changed to display a different meny or when new elements have appeared on screen that the user should know about.
+	*/
+	FScreenReaderAnnouncementInfo UserFeedback()
+	{
+		return FScreenReaderAnnouncementInfo(false, true, EScreenReaderAnnouncementPriority::Medium);
+	}
+
+private:
+	/** True if the associazted announcement should be queued if it cannot be spoken immediately or is interrupted. Else false.*/
+	bool bShouldQueue = false;
+	/** True if the associazted announcement can be intrrupted by another announcement. Else false. */
+	bool bInterruptable= true;
+	/** The priority level of the associated announcement. */
+	EScreenReaderAnnouncementPriority Priority = EScreenReaderAnnouncementPriority::Medium;
+	// @TODOAccessibility: Update timestamp to be recorded when announcements are requested, not when they are constructed. 
+	/** The timestamp for when the associated announcement was created. */
+	double Timestamp = FPlatformTime::Seconds();
+};
+
+/**
+* A screen reader announcement represents a localized message that a user wants to speak to an end user.
+* This class is used to request announcements to be spoken to an end user via the FScreenReaderUser class.
+* FScreenReaderUser::RequestSpeak() for instructions on how to use the class.
+* @see FScreenReaderUser
+*/
+USTRUCT(BlueprintType)
+struct SCREENREADER_API FScreenReaderAnnouncement
+{
+GENERATED_BODY()
+public:
+	// Only used to satisfy USTRUCT requirements
+	// Please use ctor that takes arguments
+	FScreenReaderAnnouncement() = default;
+	/**
+	* The constructor that should be used in most cases.
+	* @param InAnnouncementString The localized string that a user wishes to rquest to speak to a user
+	FScreenReaderAnnouncement(FString InAnnouncementString, FScreenReaderAnnouncementInfo );
+	* @param The announcement info associated with the announcement string to control if this announcement can be interrupted or queued
+	* when a request the speak this announcement is made.
+	*/
+	FScreenReaderAnnouncement(FString InAnnouncementString, FScreenReaderAnnouncementInfo );
+	const FString& GetAnnouncementString() const { return AnnouncementString; }
+	const FScreenReaderAnnouncementInfo& GetAnnouncementInfo() const { return AnnouncementInfo; }
+	/** Converts the contents of the announcement to a string for easy printing. Useful for debugging. */
+	FString ToString() const;
+
+	/** 
+	* The default comparator for screen reader announcements. Can be used to order screen announcements in a priority queue
+	* Users can provide their own comaprators by implementing the same interface.
+	*/
+	struct FDefaultComparator
+	{
+		/** Returns true if A is less than B */
+		bool operator()(const FScreenReaderAnnouncement& A, const FScreenReaderAnnouncement& B) const
+		{
+			// primary key is priority 
+			// secondary key is time 
+			return A.GetAnnouncementInfo().GetPriority() == B.GetAnnouncementInfo().GetPriority()
+				? A.GetAnnouncementInfo().GetTimestamp() < B.GetAnnouncementInfo().GetTimestamp()
+				: A.GetAnnouncementInfo().GetPriority() < B.GetAnnouncementInfo().GetPriority();
+		}
+	};
+
+	/**
+	* The default interruption policy to determine if announcement A can be interrupted by announcement B.
+	* Users can provide their own implementation by creating their own struct and implementing the same interface.
+	*/
+	struct FDefaultInterruptionPolicy
+	{
+		/** Returns true if the interruptee can be interrupted by the interruptor. Else returns false */
+		bool CanBeInterruptedBy(const FScreenReaderAnnouncement& Interruptee, const FScreenReaderAnnouncement& Interrupter) const
+		{
+			// priority comparison is >= because High Priority is 0 and Low Priority is 2
+			return Interruptee.GetAnnouncementInfo().IsInterruptable()
+				? Interruptee.GetAnnouncementInfo().GetPriority() >= Interrupter.GetAnnouncementInfo().GetPriority()
+				: false;
+		}
+	};
+
+	/**
+	* A static function that determines if one announcement can be interrupted by another based on a provided interruption policy.true
+	* @param Interruptee The announcement that the Interruptor will try to interrupt
+	* @param Interrupter The announcement that is trying to interrupt Interruptee
+	* @param InInterruptionPolicy The policy used to determine if the Interruptee can be interrupted by the Interrupter. Users can create their own 
+	* interruption policies that satisfy the interface provided by FDefaultInterruptionPolicy.
+	* @return Returns true if the Interruptee can be interrupted by the Interrupter based on the passed in interruption policy. Else returns false.
+	* @see FDefaultInterruptionPolicy
+	*/
+	template<typename InterruptionPolicy>
+	static bool CanBeInterruptedBy(const FScreenReaderAnnouncement& Interruptee, const FScreenReaderAnnouncement& Interrupter, const InterruptionPolicy& InInterruptionPolicy)
+	{
+		return InInterruptionPolicy.CanBeInterruptedBy(Interruptee, Interrupter);
+	}
+private:
+	/** A localized string that represents the message to be spoken to a end user */
+	FString AnnouncementString;
+	/** The announcement info assocated with the announcement which controls how the announcement behaves when a user requests this announcement to be spoken to an end user */
+	FScreenReaderAnnouncementInfo AnnouncementInfo;
+};
+
+
diff --git a/Engine/Plugins/Experimental/ScreenReader/Source/ScreenReader/Public/Announcement/ScreenReaderAnnouncementChannel.h b/Engine/Plugins/Experimental/ScreenReader/Source/ScreenReader/Public/Announcement/ScreenReaderAnnouncementChannel.h
new file mode 100644
--- /dev/null
+++ b/Engine/Plugins/Experimental/ScreenReader/Source/ScreenReader/Public/Announcement/ScreenReaderAnnouncementChannel.h
@@ -0,0 +1,48 @@
+// Copyright Epic Games, Inc. All Rights Reserved.
+
+#pragma once
+#include "CoreMinimal.h"
+#include "Announcement/ScreenReaderAnnouncement.h"
+#include "GenericPlatform/ScreenReaderReply.h"
+
+
+class FTextToSpeechBase;
+class IAccessibleWidget;
+
+/**
+* An announcement channel handles processing of announcements and speaking announcements through text to speech (TTS).
+* Announcements are forwarded here from a screen reader user and the announcements will be spoken immediately, queued or interrupt the currently spoken announcement.
+* See FScreenReaderUser for an explanation of these functions and FScreenReaderAnnouncement for how announcements are used.
+* @see FScreenReaderUser, FScreenReaderAnnouncement, FScreenReaderAnnouncementInfo
+*/
+class SCREENREADER_API FScreenReaderAnnouncementChannel
+{
+public:
+	explicit FScreenReaderAnnouncementChannel(const TSharedRef<FTextToSpeechBase>& InTextToSpeech);
+	virtual ~FScreenReaderAnnouncementChannel();
+	FScreenReaderAnnouncementChannel(const FScreenReaderAnnouncementChannel& OtherAnnouncementChannel) = delete;
+	FScreenReaderAnnouncementChannel& operator=(const FScreenReaderAnnouncementChannel& OtherAnnouncementChannel) = delete;
+
+	void Activate();
+	void Deactivate();
+
+	FScreenReaderReply RequestSpeak(FScreenReaderAnnouncement InAnnouncement);
+	FScreenReaderReply StopSpeaking();
+	bool IsSpeaking() const;
+	FScreenReaderReply ClearAllAnnouncements();
+	FScreenReaderReply RequestSpeakWidget(const TSharedRef<IAccessibleWidget>& InWidget);
+
+	
+private:
+	/**
+	* A callback that is bound to the text to speech delegate. The delegate fires when a string
+	* is fully synthesized and played to an end userwithout any interruption or stopping.
+	* This function i sused to respond to that event and speak the next queued announcement.
+	*/
+	void OnTextToSpeechFinishSpeaking();
+	/** The current announcement that is being spoken. Will have an empty string if no announcements are currently being spoken. */
+	FScreenReaderAnnouncement CurrentAnnouncement;
+	/** TheA priority queue for all of the queued announcements. */
+	TArray<FScreenReaderAnnouncement> AnnouncementQueue;
+	TSharedRef<FTextToSpeechBase> TextToSpeech;
+};
diff --git a/Engine/Plugins/Experimental/ScreenReader/Source/ScreenReader/Public/GenericPlatform/IScreenReaderBuilder.h b/Engine/Plugins/Experimental/ScreenReader/Source/ScreenReader/Public/GenericPlatform/IScreenReaderBuilder.h
new file mode 100644
--- /dev/null
+++ b/Engine/Plugins/Experimental/ScreenReader/Source/ScreenReader/Public/GenericPlatform/IScreenReaderBuilder.h
@@ -0,0 +1,39 @@
+// Copyright Epic Games, Inc. All Rights Reserved.
+
+#pragma once
+#include "CoreMinimal.h"
+
+
+class GenericApplication;
+class FScreenReaderBase;
+
+/**
+* The abstract interface for builders of all screen readers.
+* Users should subclass and provide a concrete implementation for building their custom screen readers.
+*/
+class IScreenReaderBuilder
+{
+public:
+	IScreenReaderBuilder() = default;
+	virtual ~IScreenReaderBuilder() = default;
+
+	/** The arguments to be passed to the screen reader to be built */
+	struct FArgs
+	{
+	public:
+		FArgs() = delete;
+		explicit FArgs(const TSharedRef<GenericApplication>& InPlatformApplication)
+			: PlatformApplication(InPlatformApplication)
+		{}
+		~FArgs() = default;
+		/** The platform application to be passed to the screen reader for initialization. */
+		const TSharedRef<GenericApplication>& PlatformApplication;
+	};
+	/**
+	* A pure virtual function to build a screen reader.
+	* Users must override this class to provide logic for building their custom screen readers. 
+	* @param InArgs The arguments used to build a screen reader
+	* @return The constructed screen reader built from the passed in arguments 
+	*/
+virtual TSharedRef<FScreenReaderBase> Create(const IScreenReaderBuilder::FArgs& InArgs) = 0;
+};
diff --git a/Engine/Plugins/Experimental/ScreenReader/Source/ScreenReader/Public/GenericPlatform/ScreenReaderApplicationMessageHandlerBase.h b/Engine/Plugins/Experimental/ScreenReader/Source/ScreenReader/Public/GenericPlatform/ScreenReaderApplicationMessageHandlerBase.h
new file mode 100644
--- /dev/null
+++ b/Engine/Plugins/Experimental/ScreenReader/Source/ScreenReader/Public/GenericPlatform/ScreenReaderApplicationMessageHandlerBase.h
@@ -0,0 +1,106 @@
+// Copyright Epic Games, Inc. All Rights Reserved.
+
+#pragma once
+#include "CoreMinimal.h"
+#include "GenericPlatform/GenericApplicationMessageHandler.h"
+
+class FScreenReaderBase;
+
+/**
+* The base class all screen reader application message handlers should derive from. By default, this is just a proxy for the underlying target message handler.
+* The screen reader application message handler will intercept all application messages and may perform processing on the events before consuming the event or passing it to the underlying target message handler.
+* Child classes should override the virtual event handlers from FGenericAccessibleMessageHandler to perform input processing for thier needs.
+* For example:
+* A mobile screen reader application message handler can subclass this class and intercept touch inputs to perform gesture detection to provide alternative navigation schemes for users.
+* A desktop screen reader application message handler can subclass this class and intercept mouse movement events to speak widgets that the user mouses over.
+*/
+class SCREENREADER_API FScreenReaderApplicationMessageHandlerBase : public FGenericApplicationMessageHandler
+{
+public:
+	FScreenReaderApplicationMessageHandlerBase() = delete;
+	FScreenReaderApplicationMessageHandlerBase(const TSharedRef<FGenericApplicationMessageHandler>& InTargetMessageHandler, FScreenReaderBase& InOwningScreenReader);
+	virtual ~FScreenReaderApplicationMessageHandlerBase() = default;
+
+	/** Returns true if mouse events should be processed. Else returns false. */
+	bool ShouldProcessMouseInput() const { return bProcessMouseInput; }
+	/** Sets whether mouse events should be intercepted and processed. If false, the mouse events will bubble straight down to the target message handler. */
+	void SetProcessMouseInput(bool bInProcessMouseInput) { bProcessMouseInput = bInProcessMouseInput; }
+
+	/** Retursn true if the key inputs should be intercepted and processed. Else returns false. */
+	bool ShouldProcessKeyInput() const { return bProcessKeyInput; }
+	/** Sets whether key inputs should be intercepted and processed. If set to false, the key inputs will ubble straight down to the target message handler. */
+	void SetProcessKeyInput(bool bInProcessKeyInput) { bProcessKeyInput = bInProcessKeyInput; }
+
+	/** Returns the underlying message handler that application messages bubble down to if the screen reader application message handler does not handle the event. */
+	TSharedRef<FGenericApplicationMessageHandler> GetTargetMessageHandler() const { return TargetMessageHandler; }
+
+	// FGenericApplicationMessageHandler
+	// By default all the following functions just forward the call to the underlying message handler.
+	// Child classes should only override functions where they want to intercept the message first and perform pre-processing or handling
+	// @see TargetMessageHandler
+	virtual bool ShouldProcessUserInputMessages(const TSharedPtr< FGenericWindow >& PlatformWindow) const override;
+	virtual bool OnKeyChar(const TCHAR Character, const bool IsRepeat) override;
+	virtual bool OnKeyDown(const int32 KeyCode, const uint32 CharacterCode, const bool IsRepeat) override;
+	virtual bool OnKeyUp(const int32 KeyCode, const uint32 CharacterCode, const bool IsRepeat) override;
+	virtual bool OnMouseDown(const TSharedPtr< FGenericWindow >& Window, const EMouseButtons::Type Button) override;
+	virtual bool OnMouseDown(const TSharedPtr< FGenericWindow >& Window, const EMouseButtons::Type Button, const FVector2D CursorPos) override;
+	virtual bool OnMouseUp(const EMouseButtons::Type Button) override;
+	virtual bool OnMouseUp(const EMouseButtons::Type Button, const FVector2D CursorPos) override;
+	virtual bool OnMouseDoubleClick(const TSharedPtr< FGenericWindow >& Window, const EMouseButtons::Type Button) override;
+	virtual bool OnMouseDoubleClick(const TSharedPtr< FGenericWindow >& Window, const EMouseButtons::Type Button, const FVector2D CursorPos) override;
+	virtual bool OnMouseWheel(const float Delta) override;
+	virtual bool OnMouseWheel(const float Delta, const FVector2D CursorPos) override;
+	virtual bool OnMouseMove() override;
+	virtual bool OnRawMouseMove(const int32 X, const int32 Y) override;
+	virtual bool OnCursorSet() override;
+	virtual bool OnControllerAnalog(FGamepadKeyNames::Type KeyName, int32 ControllerId, float AnalogValue) override;
+	virtual bool OnControllerButtonPressed(FGamepadKeyNames::Type KeyName, int32 ControllerId, bool IsRepeat) override;
+	virtual bool OnControllerButtonReleased(FGamepadKeyNames::Type KeyName, int32 ControllerId, bool IsRepeat) override;
+	virtual void OnBeginGesture() override;
+	virtual bool OnTouchGesture(EGestureEvent GestureType, const FVector2D& Delta, float WheelDelta, bool bIsDirectionInvertedFromDevice) override;
+	virtual void OnEndGesture() override;
+	virtual bool OnTouchStarted(const TSharedPtr< FGenericWindow >& Window, const FVector2D& Location, float Force, int32 TouchIndex, int32 ControllerId) override;
+	virtual bool OnTouchMoved(const FVector2D& Location, float Force, int32 TouchIndex, int32 ControllerId) override;
+	virtual bool OnTouchEnded(const FVector2D& Location, int32 TouchIndex, int32 ControllerId) override;
+	virtual bool OnTouchForceChanged(const FVector2D& Location, float Force, int32 TouchIndex, int32 ControllerId) override;
+	virtual bool OnTouchFirstMove(const FVector2D& Location, float Force, int32 TouchIndex, int32 ControllerId) override;
+	virtual void ShouldSimulateGesture(EGestureEvent Gesture, bool bEnable) override;
+	virtual bool OnMotionDetected(const FVector& Tilt, const FVector& RotationRate, const FVector& Gravity, const FVector& Acceleration, int32 ControllerId) override;
+	virtual bool OnSizeChanged(const TSharedRef< FGenericWindow >& Window, const int32 Width, const int32 Height, bool bWasMinimized = false) override;
+	virtual void OnOSPaint(const TSharedRef<FGenericWindow>& Window) override;
+	virtual FWindowSizeLimits GetSizeLimitsForWindow(const TSharedRef<FGenericWindow>& Window) const override;
+	virtual void OnResizingWindow(const TSharedRef< FGenericWindow >& Window) override;
+	virtual bool BeginReshapingWindow(const TSharedRef< FGenericWindow >& Window) override;
+	virtual void FinishedReshapingWindow(const TSharedRef< FGenericWindow >& Window) override;
+	virtual void HandleDPIScaleChanged(const TSharedRef< FGenericWindow >& Window) override;
+	virtual void OnMovedWindow(const TSharedRef< FGenericWindow >& Window, const int32 X, const int32 Y) override;
+	virtual bool OnWindowActivationChanged(const TSharedRef< FGenericWindow >& Window, const EWindowActivation ActivationType) override;
+	virtual bool OnApplicationActivationChanged(const bool IsActive) override;
+	virtual bool OnConvertibleLaptopModeChanged() override;
+	virtual EWindowZone::Type GetWindowZoneForPoint(const TSharedRef< FGenericWindow >& Window, const int32 X, const int32 Y) override;
+	virtual void OnWindowClose(const TSharedRef< FGenericWindow >& Window) override;
+	virtual EDropEffect::Type OnDragEnterText(const TSharedRef< FGenericWindow >& Window, const FString& Text) override;
+	virtual EDropEffect::Type OnDragEnterFiles(const TSharedRef< FGenericWindow >& Window, const TArray< FString >& Files) override;
+	virtual EDropEffect::Type OnDragEnterExternal(const TSharedRef< FGenericWindow >& Window, const FString& Text, const TArray< FString >& Files) override;
+	virtual EDropEffect::Type OnDragOver(const TSharedPtr< FGenericWindow >& Window) override;
+	virtual void OnDragLeave(const TSharedPtr< FGenericWindow >& Window) override;
+	virtual EDropEffect::Type OnDragDrop(const TSharedPtr< FGenericWindow >& Window) override;
+	virtual bool OnWindowAction(const TSharedRef< FGenericWindow >& Window, const EWindowAction::Type InActionType) override;
+	// 
+protected:
+	/** 
+	* The screen reader that owns this screen reader application message handler.
+	* The lifetime of the screen reader message handler follows its owning screen reader.
+	* The owning screen reader can be used to retrieve screen reader users to request accessible services such as requesting nnouncements be spoken
+	* or to handle accessible events.
+	* @see FScreenReaderBase
+	*/
+	FScreenReaderBase& OwningScreenReader;
+private:
+	/** The underlying message handler that application messages will bubble down to if the screen reader applicaiton message handler does not consume or handle the app.application messages. */
+	TSharedRef<FGenericApplicationMessageHandler> TargetMessageHandler;
+	/** True if the screen reader application message handler should process mouse input. Else false. */
+	bool bProcessMouseInput;
+	/** True if the screen reader application message handler should process key inputs. Else false. */
+	bool bProcessKeyInput;
+};
diff --git a/Engine/Plugins/Experimental/ScreenReader/Source/ScreenReader/Public/GenericPlatform/ScreenReaderBase.h b/Engine/Plugins/Experimental/ScreenReader/Source/ScreenReader/Public/GenericPlatform/ScreenReaderBase.h
new file mode 100644
--- /dev/null
+++ b/Engine/Plugins/Experimental/ScreenReader/Source/ScreenReader/Public/GenericPlatform/ScreenReaderBase.h
@@ -0,0 +1,107 @@
+// Copyright Epic Games, Inc. All Rights Reserved.
+
+#pragma once
+#include "CoreMinimal.h"
+
+class FGenericApplicationMessageHandler;
+class FScreenReaderApplicationMessageHandlerBase;
+class FScreenReaderUser;
+class GenericApplication;
+
+
+/**
+* The base class all screen readers must derive from.
+* The screen reader sets up and tears down the screen reader framework with Activate() and Deactivate() respectively.
+* When active, the screen reader will intercept all application messages, perform necessary input pre-processing and dispatch the accessible events and 
+* inputs to all registered screen reader users.
+* The screen reader contains a registry of screen reader users which can be registered, unregistered and retrieved with corresponding user Ids.
+* Registered users will receive and respond to accessible inputs and accessible events dispatched by the screen reader.
+* Users of this class and its child classes must activate the screen reader to ensure the accessible inputs and accessible events are dispatched properly.
+* @see FScreenReaderUser
+*/
+class SCREENREADER_API FScreenReaderBase
+{
+public: 
+	FScreenReaderBase() = delete;
+	explicit FScreenReaderBase(const TSharedRef<GenericApplication>& InPlatformApplication);
+	// We disallow copying 
+	FScreenReaderBase(const FScreenReaderBase& Other) = delete;
+	FScreenReaderBase& operator=(const FScreenReaderBase& Other) = delete;
+	virtual ~FScreenReaderBase();
+
+	void SetScreenReaderApplicationMessageHandler(const TSharedRef<FScreenReaderApplicationMessageHandlerBase>& InScreenReaderApplicationMessageHandler) { ScreenReaderApplicationMessageHandler= InScreenReaderApplicationMessageHandler; }
+	/**
+	* Activates the screen reader and sets up the screen reader framework if the screen reader is inactive.
+	* Once this is called, the screen reader can intercept application messages, perform input pre-processing and dispatch accessibility events to the registered screen reader users.
+	* OnActivated() is called at the end of the function. Users can override the function to perform additional setup.
+	* @see OnActivated()
+	*/
+	void Activate();
+	/**
+	* Deactivates the screen reader and tears down the screen reading framework if the screen reader is active.
+	* The screen reader will no longer intercept application input messages and no accessible events will be dispatched to the registered screen reader users. 
+	* OnDeactivated() is called at the start of this function. Users can override the function to perform additional teardown.
+	* @see OnDeactivated()
+	*/
+	void Deactivate();
+	/** Returns true if the screen reader is active. Else false. */
+	bool IsActive() const { return bActive; }
+	/** 
+	* Registers a provided user Id to the screen reader framework and allows the user to receive and respond to accessible events and accessible input.
+	* Does nothing if the passed in Id is already registered.
+	* @ return Returns true if the Id is successfully registered with the screen reader. Else returns false.
+	*/
+	bool RegisterUser(int32 InUserId);
+	/**
+	* Unregisters a provided user Id from the screen reader framework and deactivates the user. The unregistered user will no longer receive or respond to accessible events and input.
+	* Nothing will happen if the provided user Id has not been registered with the screen reader. 
+	* @param InUserId The user Id to unregister from the screen reader framework.
+	* @return Returns true if the provided user Id has been successfully unregistered. Else returns false.
+	*/
+	bool UnregisterUser(int32 InUserId);
+	/**
+	* Returns true if the passed in user Id is already registered. Else returns false.
+	* @param InUserId The user Id to check for registration.
+	*/
+	bool IsUserRegistered(int32 InUserId) const;
+	/** Unregisteres all screen reader users. No users can receive or respond to accessible events or accessible input upon calling this function. */
+	void UnregisterAllUsers();
+	/** 
+	* Returns the screen reader user associated with the passed in Id. It is up to the user
+	* to ensure that the Id is valid and mapped to a valid screen reader user.
+	* @param InUserId The Id of the screen reader user to be retrieved.
+	* @return The screen reader user associated with the passed in Id
+	*/
+	TSharedRef<FScreenReaderUser> GetUser(int32 InUserId) const;
+	/**
+	* Returns a screen reader user associated with the passed in Id if the Id is registered. Otherwise, a nullptr is returned.
+	* @param InUserId The user Id of the screen reader to be retrieved.
+	* @return The screen reader user associated with the Id if the user is registered. Else false.
+	*/
+	TSharedPtr<FScreenReaderUser> GetUserChecked(int32 InUserId) const;
+
+protected:
+	/**
+	* Called at the end of Activate(). Override this function in child classes to do additional set up 
+	* such as subscribing to accessible events, binding handlers for focus events or providing feedback to
+	* users that the screen reader has been activated.
+	*/
+	virtual void OnActivate() {}
+	/**
+	* Called at the start of Deactivate(). Override this function in child classes to provide additional teardown such as unsubscribing from accessible events or providing
+	* feedback to alert the user that the screen reader is deactivated. 
+	*/
+	virtual void OnDeactivate() {}
+private:
+	/** 
+	* The screen reader's application message handler. It processes input events before passing them down to the target application message handler
+	* @see FScreenReaderApplicationMessageHandlerBase
+	*/
+	TSharedRef<FScreenReaderApplicationMessageHandlerBase> ScreenReaderApplicationMessageHandler;
+	/** The platform application. Used to access the accessible message handler and swap the application message handler with the screen reader application message handler.. */
+	TWeakPtr<GenericApplication> PlatformApplication;
+	/** A map of user Ids to screen reader users. Each entry represents a registered user that can receive and repsond to accessible events and accessible input. */
+	TMap<int32, TSharedRef<FScreenReaderUser>> UsersMap;
+	/** True if the screen reader is active. Else false. */
+	bool bActive;
+};
\ No newline at end of file
diff --git a/Engine/Plugins/Experimental/ScreenReader/Source/ScreenReader/Public/GenericPlatform/ScreenReaderReply.h b/Engine/Plugins/Experimental/ScreenReader/Source/ScreenReader/Public/GenericPlatform/ScreenReaderReply.h
new file mode 100644
--- /dev/null
+++ b/Engine/Plugins/Experimental/ScreenReader/Source/ScreenReader/Public/GenericPlatform/ScreenReaderReply.h
@@ -0,0 +1,38 @@
+// Copyright Epic Games, Inc. All Rights Reserved.
+
+#pragma once
+#include "CoreMinimal.h"
+/**
+* A struct passed around the the screen reader framework to indicate if a request has been successfully handled.
+* Use this class to determine if a request to the screen reader user was successfully handled and provide user feedback accordingly.
+* // @TODOAccessibility: Provide examples 
+* @see FScreenReaderUser
+*/
+class SCREENREADER_API FScreenReaderReply
+{
+public:
+	FScreenReaderReply() = delete;
+	~FScreenReaderReply() = default;
+
+	static FScreenReaderReply Handled()
+	{
+		return FScreenReaderReply(true);
+	}
+
+	static FScreenReaderReply Unhandled()
+	{
+		return FScreenReaderReply(false);
+	}
+
+	bool IsHandled() const
+	{
+		return bHandled;
+	}
+private:
+	FScreenReaderReply(bool bInHandled)
+		: bHandled(bInHandled)
+	{
+
+	}
+	bool bHandled;
+};
\ No newline at end of file
diff --git a/Engine/Plugins/Experimental/ScreenReader/Source/ScreenReader/Public/GenericPlatform/ScreenReaderUser.h b/Engine/Plugins/Experimental/ScreenReader/Source/ScreenReader/Public/GenericPlatform/ScreenReaderUser.h
new file mode 100644
--- /dev/null
+++ b/Engine/Plugins/Experimental/ScreenReader/Source/ScreenReader/Public/GenericPlatform/ScreenReaderUser.h
@@ -0,0 +1,92 @@
+// Copyright Epic Games, Inc. All Rights Reserved.
+
+#pragma once
+#include "CoreMinimal.h"
+#include "Announcement/ScreenReaderAnnouncement.h"
+#include "GenericPlatform/ScreenReaderReply.h"
+
+class IAccessibleWidget;
+class FScreenReaderAnnouncementChannel;
+
+/**
+* A user of the screen reader. Corresponds to a hardware device that users use.
+* This class is a facade that acts as a one stop shop for all screen reading services on a per user basis.
+* Multiple screen reader users can exist simultaneoulsy to facilitate local multiplayer and users can opt in or out 
+* of receiving screen reader feedback by registering and unregistering with the screen reader respectively.
+* A screen user is inactive by default when it is first registered with teh screen reader. Users must explicitly activate the screen reader user after successful registration with the screen reader to use its services.
+* Responsibilities of the class are:
+* 1. Text to speech (TTS) requests - Users can request an announcement to be spoken via text to speech to  screen reader user
+* 2. Accessible focus handling - A screen reader user holds information about the accessible widget it is currently focused on
+* @see FScreenReaderBase, FScreenReaderAnnouncement
+*/
+class SCREENREADER_API FScreenReaderUser
+{
+public:
+	explicit FScreenReaderUser(int32 InUserId);
+	~FScreenReaderUser();
+	/**
+	* Requests an announcement to be spoken to the screen reader user. This is the main mechanism to provide text to speech auditory feedback
+	* to users.
+	*Calling this function does not guarantee that the announcement will be spoken via text to speech and be 
+	* heard by a user.
+	* All announcements spoken via text to speech will be asynchronous and iwll not block the game thread. 
+	* If the screen reader user is active and no announcements are currently spoken, the announcement will be spoken immediately.
+	* If another announcement is currently being spoken, the passed in announcement could be queued or interrupt the currently spoken announcement.
+	* Examples:
+	* TSharedRef<FScreenReaderUser> MyUser = ScreenReader->GetUser(MyScreenReaderUserId);
+	* // All announcements to be spoken should be localized to provide language support for the text to speech system.
+	* static const FText MyText = LOCTEXT("ExampleUserFeedback", "Feedback to user.");
+	* // This makes the announcement interrupt any currently spoken announcement with lower priority and is played immediately. The announcement will be uninterruptable and guarantees the user will hear the announcement.
+	* // Use this sparingly for only critical information that the user needs to hear such as registration to the screen reader or important in game events.
+	* MyUser->RequestSpeak(FScreenReaderAnnouncement(MyText.ToString(), FScreenReaderAnnouncementInfo::Important()));
+	* // This makes the announcement a feedback for a user action or event. It is interruptable.
+	* // Use this to alert users to an action they performed 
+	* MyUser->RequestSpeak(FScreenReaderAnnouncement(MyText.ToString(), FScreenReaderAnnouncementInfo::UserFeedback()));
+	* @param InAnnouncement - The announcement requested to be spoken to the screen reader user
+	* @return FScreenReaderReply::Handled() if the request was successfully processed. Else return FScreenReader::Unhandled()
+	* @see FScreenReaderAnnouncement, FScreenReaderAnnouncementInfo
+	*/
+	FScreenReaderReply RequestSpeak(FScreenReaderAnnouncement InAnnouncement);
+	/**
+	* Immediately stops speaking any currently spoken announcement.
+	* Does nothing if there is no announcement currently being spoken. 
+	* @return FScreenReader::Handled() if the request to stop speaking is successfully processed. Else returns FScreenReaderReply::Unhandled()
+	*/
+	FScreenReaderReply StopSpeaking();
+	/** Returns true if an announcement is currently being spoken. Else returns false.*/
+	bool IsSpeaking() const;
+	/**
+	* Requests the information about an accessibility widget to be read out to a user.
+	* The same guarantees about the announcement being spoken in RequestSpeak() apply for this function.
+	* @param InWidget The widget to ahve its contents spoken to a user.
+	* @return FScreenReaderReply::Handled() if the widget's information is successfully spoken. Else returns FScreenReaderReply::Unhandled();
+	*/
+	FScreenReaderReply RequestSpeakWidget(const TSharedRef<IAccessibleWidget>& InWidget);
+
+	/**
+	* Activates the screen reader and allow users to use its services such as requesting announcements to be spoken via text to speech.
+	* When screen reader users are firstr registered with a screen reader, they are inactive by default. Users must explicitly activate the screen reader user.
+	*/
+	void Activate();
+	/** 
+	* Deactivates the screen reader and disables all announcement and text to speech services
+	* making them do nothing.
+	*/
+
+	void Deactivate();
+	/** Returns true if the screen reader user is active. Else returns false.*/
+	bool IsActive() const { return bActive; }
+
+	/** Retursn the user Id associated with this screen reader user. */
+	int32 GetUserId() const { return UserId; }
+	/** Returns the accessible widget the screen reader user is currently focused on. */
+	TSharedPtr<IAccessibleWidget> GetAccessibleFocus() const { return AccessibleFocus.Pin(); }
+	/** Sets the accessible the screenr eader user is currently focused on */
+	void SetAccessibleFocus(const TSharedRef<IAccessibleWidget>& InAccessibleFocus) { AccessibleFocus = InAccessibleFocus; }
+private:
+	int32 UserId;
+	/** Responsible for handling all incoming announcement requests and speaking them via text to speech if possible */
+	TUniquePtr<FScreenReaderAnnouncementChannel> AnnouncementChannel;
+	TWeakPtr<IAccessibleWidget> AccessibleFocus;
+	bool bActive;
+};
\ No newline at end of file
diff --git a/Engine/Plugins/Experimental/ScreenReader/Source/ScreenReader/ScreenReader.Build.cs b/Engine/Plugins/Experimental/ScreenReader/Source/ScreenReader/ScreenReader.Build.cs
new file mode 100644
--- /dev/null
+++ b/Engine/Plugins/Experimental/ScreenReader/Source/ScreenReader/ScreenReader.Build.cs
@@ -0,0 +1,22 @@
+// Copyright Epic Games, Inc. All Rights Reserved.
+using UnrealBuildTool;
+
+
+public class ScreenReader : ModuleRules
+{
+	public ScreenReader(ReadOnlyTargetRules Target) : base(Target)
+	{
+		PublicDependencyModuleNames.AddRange(
+			new string[] {
+			"Core",
+			"CoreUObject"
+		});
+
+		PrivateDependencyModuleNames.AddRange(
+			new string[] {
+			"ApplicationCore",
+			"TextToSpeech",
+			"Engine"
+		});
+	}
+}
diff --git a/Engine/Plugins/Experimental/SlateScreenReader/SlateScreenReader.uplugin b/Engine/Plugins/Experimental/SlateScreenReader/SlateScreenReader.uplugin
new file mode 100644
--- /dev/null
+++ b/Engine/Plugins/Experimental/SlateScreenReader/SlateScreenReader.uplugin
@@ -0,0 +1,40 @@
+{
+	"FileVersion" : 3,
+	"Version" : 1,
+	"VersionName" : "0.1",
+	"FriendlyName" : "SlateScreenReader",
+	"Description" : "A screen reader that provides vision accessibility services for Slate.",
+	"Category" : "Accessibility",
+	"CreatedBy" : "Epic Games, Inc.",
+	"CreatedByURL" : "http://epicgames.com",
+	"DocsURL" : "",
+	"MarketplaceURL" : "",
+	"SupportURL" : "",
+	"EnabledByDefault" : false,
+	"CanContainContent" : false,
+	"IsExperimental" : true,
+	"Installed" : false,
+	"Modules" :
+	[
+		{
+			"Name" : "SlateScreenReader",
+			"Type" : "Runtime",
+			"LoadingPhase" : "PostEngineInit",
+			"WhitelistPlatforms" :
+			[
+				"Win64"
+			],
+			"BlacklistTargets" :
+			[
+				"Server"
+			]
+		}
+	],
+	"Plugins" :
+	[
+		{
+			"Name" : "ScreenReader",
+			"Enabled": true
+		}
+	]
+}
diff --git a/Engine/Plugins/Experimental/SlateScreenReader/Source/SlateScreenReader/Private/SlateScreenReader.cpp b/Engine/Plugins/Experimental/SlateScreenReader/Source/SlateScreenReader/Private/SlateScreenReader.cpp
new file mode 100644
--- /dev/null
+++ b/Engine/Plugins/Experimental/SlateScreenReader/Source/SlateScreenReader/Private/SlateScreenReader.cpp
@@ -0,0 +1,58 @@
+// Copyright Epic Games, Inc. All Rights Reserved.
+
+#include "SlateScreenReader.h"
+#include "GenericPlatform/GenericApplication.h"
+#include "GenericPlatform/Accessibility/GenericAccessibleInterfaces.h"
+#include "Announcement/ScreenReaderAnnouncementChannel.h"
+#include "GenericPlatform/ScreenReaderUser.h"
+#include "InputCoreTypes.h"
+#include "SlateScreenReaderLog.h"
+
+// Slate includes
+#include "Framework/Application/SlateApplication.h"
+#include "Framework/Application/SlateUser.h"
+#include "Input/Events.h"
+#include "Widgets/Accessibility/SlateAccessibleWidgetCache.h"
+
+#define LOCTEXT_NAMESPACE "SlateScreenReader"
+
+FSlateScreenReader::FSlateScreenReader(const TSharedRef<GenericApplication>& InPlatformApplication)
+	: FScreenReaderBase(InPlatformApplication)
+{
+
+}
+
+FSlateScreenReader::~FSlateScreenReader()
+{
+
+}
+
+void FSlateScreenReader::HandleSlateFocusChanging(const FFocusEvent& InFocusEvent, const FWeakWidgetPath& InOldFocusWidgetPath, const TSharedPtr<SWidget>& InOldWidget, const FWidgetPath& InNewWidgetPath, const TSharedPtr<SWidget>& InNewFocusWidget)
+{
+	if (IsActive() && InNewFocusWidget && InNewFocusWidget->IsAccessible())
+	{
+		int32 UserId = static_cast<int32>(InFocusEvent.GetUser());
+		if (TSharedPtr<FScreenReaderUser> User = GetUserChecked(UserId))
+		{
+			TSharedPtr<IAccessibleWidget> AccessibleWidget = FSlateAccessibleWidgetCache::GetAccessibleWidgetChecked(InNewFocusWidget);
+			if (AccessibleWidget)
+			{
+				User->SetAccessibleFocus(AccessibleWidget.ToSharedRef());
+				User->RequestSpeakWidget(AccessibleWidget.ToSharedRef());
+			}
+		}
+	}
+}
+
+void FSlateScreenReader::OnActivate()
+{
+	UE_LOG(LogSlateScreenReader, Verbose, TEXT("On activation of Slate screen reader."));
+	// @TODOAccessibility: Consider removing the binding on deactivation 
+	FSlateApplication::Get().OnFocusChanging().AddRaw(this, &FSlateScreenReader::HandleSlateFocusChanging);
+}
+
+void FSlateScreenReader::OnDeactivate()
+{
+	UE_LOG(LogSlateScreenReader, Verbose, TEXT("Deactivation of Slate screen reader."));
+} 
+#undef LOCTEXT_NAMESPACE
diff --git a/Engine/Plugins/Experimental/SlateScreenReader/Source/SlateScreenReader/Private/SlateScreenReader.h b/Engine/Plugins/Experimental/SlateScreenReader/Source/SlateScreenReader/Private/SlateScreenReader.h
new file mode 100644
--- /dev/null
+++ b/Engine/Plugins/Experimental/SlateScreenReader/Source/SlateScreenReader/Private/SlateScreenReader.h
@@ -0,0 +1,34 @@
+// Copyright Epic Games, Inc. All Rights Reserved.
+
+#pragma once
+
+#include "CoreMinimal.h"
+#include "GenericPlatform/ScreenReaderBase.h"
+#include "GenericPlatform/GenericApplicationMessageHandler.h"
+
+class SWidget;
+class iAccessibleWidget;
+class FWidgetPath;
+class FWeakWidgetPath;
+struct FFocusEvent;
+
+/**
+ * A basic screen reader that can work for desktop and consoles that use Slate.
+ * All feedback to users are done through speech.
+ */
+class FSlateScreenReader
+	: public FScreenReaderBase
+{
+public:
+	FSlateScreenReader() = delete;
+	FSlateScreenReader(const FSlateScreenReader& Other) = delete;
+	explicit FSlateScreenReader(const TSharedRef<GenericApplication>& InPlatformApplication);
+
+	FSlateScreenReader& operator=(const FSlateScreenReader& Other) = delete;
+	virtual ~FSlateScreenReader();
+	
+	void HandleSlateFocusChanging(const FFocusEvent& InFocusEvent, const FWeakWidgetPath& InOldFocusWidgetPath, const TSharedPtr<SWidget>& InOldWidget, const FWidgetPath& InNewWidgetPath, const TSharedPtr<SWidget>& InNewFocusWidget);
+protected:
+	virtual void OnActivate() override;
+	virtual void OnDeactivate() override;
+};
diff --git a/Engine/Plugins/Experimental/SlateScreenReader/Source/SlateScreenReader/Private/SlateScreenReaderBuilder.cpp b/Engine/Plugins/Experimental/SlateScreenReader/Source/SlateScreenReader/Private/SlateScreenReaderBuilder.cpp
new file mode 100644
--- /dev/null
+++ b/Engine/Plugins/Experimental/SlateScreenReader/Source/SlateScreenReader/Private/SlateScreenReaderBuilder.cpp
@@ -0,0 +1,13 @@
+// Copyright Epic Games, Inc. All Rights Reserved.
+
+#include "SlateScreenReaderBuilder.h"
+#include "GenericPlatform/ScreenReaderBase.h"
+#include "SlateScreenReader.h"
+#include "GenericPlatform/GenericApplication.h"
+
+TSharedRef<FScreenReaderBase> FSlateScreenReaderBuilder::Create(const IScreenReaderBuilder::FArgs& InArgs)
+{
+	TSharedRef<FScreenReaderBase> ScreenReader = MakeShared<FSlateScreenReader>(InArgs.PlatformApplication);
+	// @TODOAccessibility: Set custom screen reader application message handlers for different input handling 
+	return ScreenReader;
+}
diff --git a/Engine/Plugins/Experimental/SlateScreenReader/Source/SlateScreenReader/Private/SlateScreenReaderBuilder.h b/Engine/Plugins/Experimental/SlateScreenReader/Source/SlateScreenReader/Private/SlateScreenReaderBuilder.h
new file mode 100644
--- /dev/null
+++ b/Engine/Plugins/Experimental/SlateScreenReader/Source/SlateScreenReader/Private/SlateScreenReaderBuilder.h
@@ -0,0 +1,16 @@
+// Copyright Epic Games, Inc. All Rights Reserved.
+
+#pragma once
+
+#include "CoreMinimal.h"
+#include "GenericPlatform/IScreenReaderBuilder.h"
+
+class FSlateScreenReaderBuilder : public IScreenReaderBuilder
+{
+public:
+	FSlateScreenReaderBuilder() {}
+	virtual ~FSlateScreenReaderBuilder() {}
+
+	virtual TSharedRef<FScreenReaderBase> Create(const IScreenReaderBuilder::FArgs& InArgs) override;
+
+};
diff --git a/Engine/Plugins/Experimental/SlateScreenReader/Source/SlateScreenReader/Private/SlateScreenReaderEngineSubsystem.cpp b/Engine/Plugins/Experimental/SlateScreenReader/Source/SlateScreenReader/Private/SlateScreenReaderEngineSubsystem.cpp
new file mode 100644
--- /dev/null
+++ b/Engine/Plugins/Experimental/SlateScreenReader/Source/SlateScreenReader/Private/SlateScreenReaderEngineSubsystem.cpp
@@ -0,0 +1,85 @@
+// Copyright Epic Games, Inc. All Rights Reserved.
+
+#include "SlateScreenReaderEngineSubsystem.h"
+#include "GenericPlatform/ScreenReaderBase.h"
+#include "Application/SlateApplicationBase.h"
+#include "GenericPlatform/IScreenReaderBuilder.h"
+#include "SlateScreenReaderModule.h"
+#include "Framework/Application/SlateApplication.h"
+#include "Engine/Engine.h"
+
+
+USlateScreenReaderEngineSubsystem::USlateScreenReaderEngineSubsystem()
+{
+
+}
+
+USlateScreenReaderEngineSubsystem::~USlateScreenReaderEngineSubsystem()
+{
+	ensureMsgf(!ScreenReader, TEXT("Screen reader should already be destroyed and nulled at this point."));
+}
+
+USlateScreenReaderEngineSubsystem& USlateScreenReaderEngineSubsystem::Get()
+{
+	return *GEngine->GetEngineSubsystem<USlateScreenReaderEngineSubsystem>();
+}
+
+void USlateScreenReaderEngineSubsystem::Initialize(FSubsystemCollectionBase& Collection)
+{
+	Super::Initialize(Collection);
+	// By the time this is called, Slate is all set up
+	// I.e Platform applications already have subscribed to the delegate in the accessible
+	// message handler. So any platform stuff e.g OSX trees and cache are all set up.  That needs to get nuked somehow when we subscribe to accessible events here
+	FSlateApplication& SlateApplication = FSlateApplication::Get();
+	SlateApplication.GetPlatformApplication()->GetAccessibleMessageHandler()->UnbindAccessibleEventDelegate();
+	// @TODOAccessibility: We've unbound from the event, but we need to clean up the OS resources that are allocated for accessibility as well. 
+
+	// @TODOAccessibility: Consider lazy initialization 
+	IScreenReaderBuilder::FArgs Args(SlateApplication.GetPlatformApplication().ToSharedRef());
+	// @TODOAccessibility: Allow a means of using custom screen reader builder 
+	ScreenReader = ISlateScreenReaderModule::Get().GetDefaultScreenReaderBuilder()->Create(Args);
+	// Slate could get shutdown prior to engine subsystems 
+	SlateApplication.OnPreShutdown().AddSP(ScreenReader.ToSharedRef(), &FScreenReaderBase::Deactivate);
+}
+
+void USlateScreenReaderEngineSubsystem::Deinitialize()
+{
+	// should still be valid
+	check(ScreenReader);
+	// Engine subsystems are desroyed before FSlateApplication::Shutdown is called
+	// We clean up here 
+	if (ScreenReader->IsActive())
+	{
+		ScreenReader->Deactivate();
+	}
+	ScreenReader.Reset();
+	Super::Deinitialize();
+}
+
+bool USlateScreenReaderEngineSubsystem::ShouldCreateSubsystem(UObject* Outer) const
+{
+	// On certain builds, Slate is not initialize, we don't need the screen reader in that case 
+	return FSlateApplication::IsInitialized();
+}
+
+void USlateScreenReaderEngineSubsystem::Activate()
+{
+	if (!ScreenReader->IsActive())
+	{
+		ScreenReader->Activate();
+	}
+}
+
+void USlateScreenReaderEngineSubsystem::Deactivate()
+{
+	checkf(ScreenReader, TEXT("Invalid screen reader on deactivation. A valid screen reader must always be available for deactivation."));
+	if (ScreenReader->IsActive())
+	{
+		ScreenReader->Deactivate();
+	}
+}
+
+TSharedRef<FScreenReaderBase> USlateScreenReaderEngineSubsystem::GetScreenReader() const
+{
+	return ScreenReader.ToSharedRef();
+}
\ No newline at end of file
diff --git a/Engine/Plugins/Experimental/SlateScreenReader/Source/SlateScreenReader/Private/SlateScreenReaderLog.h b/Engine/Plugins/Experimental/SlateScreenReader/Source/SlateScreenReader/Private/SlateScreenReaderLog.h
new file mode 100644
--- /dev/null
+++ b/Engine/Plugins/Experimental/SlateScreenReader/Source/SlateScreenReader/Private/SlateScreenReaderLog.h
@@ -0,0 +1,6 @@
+// Copyright Epic Games, Inc. All Rights Reserved.
+
+#pragma once
+#include "CoreMinimal.h"
+
+DECLARE_LOG_CATEGORY_EXTERN(LogSlateScreenReader, Log, All);
diff --git a/Engine/Plugins/Experimental/SlateScreenReader/Source/SlateScreenReader/Private/SlateScreenReaderModule.cpp b/Engine/Plugins/Experimental/SlateScreenReader/Source/SlateScreenReader/Private/SlateScreenReaderModule.cpp
new file mode 100644
--- /dev/null
+++ b/Engine/Plugins/Experimental/SlateScreenReader/Source/SlateScreenReader/Private/SlateScreenReaderModule.cpp
@@ -0,0 +1,53 @@
+// Copyright Epic Games, Inc. All Rights Reserved.
+
+#include "SlateScreenReaderModule.h"
+#include "SlateScreenReaderLog.h"
+#include "SlateScreenReaderBuilder.h"
+
+
+
+class FSlateScreenReaderModule : public ISlateScreenReaderModule
+{
+public:
+	FSlateScreenReaderModule()
+		: DefaultBuilder(MakeShared<FSlateScreenReaderBuilder>())
+	{
+
+	}
+	virtual ~FSlateScreenReaderModule() = default;
+	// ISlateScreenReaderModule
+	virtual TSharedRef<IScreenReaderBuilder> GetDefaultScreenReaderBuilder() const override
+	{
+		return DefaultBuilder;
+	}
+
+	virtual TSharedPtr<IScreenReaderBuilder> GetCustomScreenReaderBuilder() const override
+	{
+		return CustomBuilder;
+}
+	
+	virtual void SetCustomScreenReaderBuilder(const TSharedRef<IScreenReaderBuilder>& InBuilder) override
+	{
+		CustomBuilder= InBuilder;
+	}
+	// ~
+
+	// IModuleInterface
+	virtual void StartupModule() override
+	{
+		
+	}
+
+	virtual void ShutdownModule() override
+	{
+		
+	}
+	// ~
+private:
+	TSharedRef<IScreenReaderBuilder> DefaultBuilder;
+	TSharedPtr<IScreenReaderBuilder> CustomBuilder;
+};
+
+IMPLEMENT_MODULE(FSlateScreenReaderModule, SlateScreenReader);
+
+DEFINE_LOG_CATEGORY(LogSlateScreenReader);
\ No newline at end of file
diff --git a/Engine/Plugins/Experimental/SlateScreenReader/Source/SlateScreenReader/Public/SlateScreenReaderEngineSubsystem.h b/Engine/Plugins/Experimental/SlateScreenReader/Source/SlateScreenReader/Public/SlateScreenReaderEngineSubsystem.h
new file mode 100644
--- /dev/null
+++ b/Engine/Plugins/Experimental/SlateScreenReader/Source/SlateScreenReader/Public/SlateScreenReaderEngineSubsystem.h
@@ -0,0 +1,63 @@
+// Copyright Epic Games, Inc. All Rights Reserved.
+
+#pragma once
+#include "CoreMinimal.h"
+#include "Subsystems/EngineSubsystem.h"
+#include "SlateScreenReaderEngineSubsystem.generated.h"
+
+class FScreenReaderBase;
+/**
+* The engine susbsystem for the Slate screen reader.
+* This class should be the entryway for C++ programmers and BP useres alike to interact with the screen reader system.
+* The subsystem must be activated before the screen reader services can be used.
+* For C++ users, please retrieve the screen reader and interact with the screen reader users from there.
+* Example:
+* USlateScreenReaderEngineSubsystem ::Get().Activate();
+* // Registers a screen reader user with Id 0. A screen reader user should correspond to a hardware input device such as a keyboard or controller like FSlateUser
+* USlateScreenReaderEngineSubsystem ::Get().GetScreenReader()->RegisterUser(0);
+* TSharedRef<FScreenReaderUser> User = USlateScreenReaderEngineSubsystem::Get().GetScreenReader()->GetUser(0);
+* // Screen reader users are inactive when they are first registered and need to be explicitly activated.
+* User->Activate();
+* static const FText HelloWorld = LOCTEXT("HelloWorld", "Hello World");
+* // Requests "Hello World" to be spoken to the screen reader user
+* User->RequestSpeak(FScreenReaderAnnouncement(HelloWorld.ToString(), FScreenReaderInfo::Important())); 
+* @see FScreenReaderBase, FScreenReaderUser, FScreenReaderAnnouncement
+*/
+UCLASS()
+class SLATESCREENREADER_API USlateScreenReaderEngineSubsystem : public UEngineSubsystem
+{
+	GENERATED_BODY()
+public:
+	USlateScreenReaderEngineSubsystem();
+	virtual ~USlateScreenReaderEngineSubsystem();
+
+	/** Convenience method to retrieve the screen reader engine subsystem. */
+	static USlateScreenReaderEngineSubsystem& Get();
+	
+	/** 
+	* Activates the underlying screen reader. Use this to allow end users to register with the screen reader 
+	* and receive accessible feedback via text to speech and get access to other screen reader services.
+	*/
+	UFUNCTION(BlueprintCallable, Category="ScreenReader")
+	void Activate();
+	/**
+	* Deactivates the underlying screen reader and prevents end users from getting
+	* any accessible feedback via text to speech or using any other screen reader services.
+	*/
+	UFUNCTION(BlueprintCallable, Category="ScreenReader")
+	void Deactivate();
+	/** Returns the underlying screen reader */
+	TSharedRef<FScreenReaderBase> GetScreenReader() const;
+
+// UEngineSubsystem
+	virtual void Initialize(FSubsystemCollectionBase& Collection) override;
+	virtual void Deinitialize() override;
+	virtual bool ShouldCreateSubsystem(UObject* Outer) const override;
+	// ~
+private:
+	/** 
+	* The underlying screen reader. This should be a TUniquePtr
+	* but we make it a TSharedPtr to allow for easy delegate unbinding without needing to manually unbind from delegates.
+	*/
+	TSharedPtr<FScreenReaderBase> ScreenReader;
+};
diff --git a/Engine/Plugins/Experimental/SlateScreenReader/Source/SlateScreenReader/Public/SlateScreenReaderModule.h b/Engine/Plugins/Experimental/SlateScreenReader/Source/SlateScreenReader/Public/SlateScreenReaderModule.h
new file mode 100644
--- /dev/null
+++ b/Engine/Plugins/Experimental/SlateScreenReader/Source/SlateScreenReader/Public/SlateScreenReaderModule.h
@@ -0,0 +1,29 @@
+// Copyright Epic Games, Inc. All Rights Reserved.
+
+#pragma once
+#include "CoreMinimal.h"
+#include "Modules/ModuleManager.h"
+#include "Modules/ModuleInterface.h"
+
+static_assert(WITH_ACCESSIBILITY, TEXT("Trying to use the Slate screen reader plugin with accessibility disabled. Accessibility must be enabled to use this plugin. Either enable accessibility or disable the plugin."));
+
+class IScreenReaderBuilder;
+
+class SLATESCREENREADER_API ISlateScreenReaderModule: public IModuleInterface
+{
+public:
+	virtual ~ISlateScreenReaderModule() {}
+	/** Get this module and load it if required. */
+	static ISlateScreenReaderModule& Get()
+	{
+		static const FName ModuleName = "SlateScreenReader";
+		return FModuleManager::Get().LoadModuleChecked<ISlateScreenReaderModule>(ModuleName);
+	}
+/** Returns the default IScreenReaderBuilder used to construct Slate screen readers. There will always be a valid screen reader builder for Slate screen readers.*/
+	virtual TSharedRef<IScreenReaderBuilder> GetDefaultScreenReaderBuilder() const = 0;
+	/** Gets the custom screen reader builder set by the user for building Slate screen readers. If the user did not set a custom screen reader builder, this returns nullptr. */
+	virtual TSharedPtr<IScreenReaderBuilder> GetCustomScreenReaderBuilder() const = 0;
+	/** Sets a custom screen reader builder to build Slate screen readers. */
+	virtual void SetCustomScreenReaderBuilder(const TSharedRef<IScreenReaderBuilder>& InBuilder) = 0;
+};
+
diff --git a/Engine/Plugins/Experimental/SlateScreenReader/Source/SlateScreenReader/SlateScreenReader.Build.cs b/Engine/Plugins/Experimental/SlateScreenReader/Source/SlateScreenReader/SlateScreenReader.Build.cs
new file mode 100644
--- /dev/null
+++ b/Engine/Plugins/Experimental/SlateScreenReader/Source/SlateScreenReader/SlateScreenReader.Build.cs
@@ -0,0 +1,26 @@
+// Copyright Epic Games, Inc. All Rights Reserved.
+using UnrealBuildTool;
+
+
+public class SlateScreenReader : ModuleRules
+{
+	public SlateScreenReader(ReadOnlyTargetRules Target) : base(Target)
+	{
+		PublicDependencyModuleNames.AddRange(
+			new string[] {
+			"Core",
+			"CoreUObject",
+			"Engine"
+		});
+
+		PrivateDependencyModuleNames.AddRange(
+			new string[] {
+			"ApplicationCore",
+			"SlateCore",
+			"Slate",
+			"InputCore",
+			"TextToSpeech",
+			"ScreenReader"
+		});
+	}
+}
diff --git a/Engine/Source/Runtime/ApplicationCore/Public/GenericPlatform/Accessibility/GenericAccessibleInterfaces.h b/Engine/Source/Runtime/ApplicationCore/Public/GenericPlatform/Accessibility/GenericAccessibleInterfaces.h
--- a/Engine/Source/Runtime/ApplicationCore/Public/GenericPlatform/Accessibility/GenericAccessibleInterfaces.h
+++ b/Engine/Source/Runtime/ApplicationCore/Public/GenericPlatform/Accessibility/GenericAccessibleInterfaces.h
@@ -519,10 +519,7 @@ public:
 
 	virtual ~FGenericAccessibleMessageHandler()
 	{
-		if (AccessibleEventDelegate.IsBound())
-		{
-			AccessibleEventDelegate.Unbind();
-		}
+		UnbindAccessibleEventDelegate();
 	}
 
 	/**
@@ -596,6 +593,14 @@ public:
 		AccessibleEventDelegate = Delegate;
 	}
 
+	/**
+ * Unbinds the delegate called during accessible events if it is bound. 
+ */
+	void UnbindAccessibleEventDelegate()
+	{
+		AccessibleEventDelegate.Unbind();
+	}
+
 	/**
 	 * Request a function to be run in a particular thread. 
 	 * Primarily used to service accessibility requests from the OS thread about accessible Slate widget data in the game thread. 
