Screen reader plugins refactor to use the new accessible focus API.
- Slate screen reader no longer subscribes to focus change events from Slate application and instead processes focus change accessible events from FGenericAccessibleMessageHandler to
provide the same behavior.
- FScreenReaderBase is now abstract and derived classes must implement OnAccessibleEventRaised to handle accessible events from FGenericAccessibleMessageHandler
- FScreenReaderBase no longer keeps track of all of the registered screen reader users and instead delegates that to the FGenericAccessibleUserRegistry.
- FScreenReaderUser no longer keeps track of all the focus data and instead delegates that to FGenericAccessibleUser
- When the Slate screen reader is constructed, all accessible users will be unregistered. This forces all users to re-register users through the screen reader framework.
-All screen reader user registration of users is forwarded to FGenericAccessibleUserRegistry

diff --git a/Engine/Plugins/Experimental/ScreenReader/Source/ScreenReader/Private/GenericPlatform/ScreenReaderBase.cpp b/Engine/Plugins/Experimental/ScreenReader/Source/ScreenReader/Private/GenericPlatform/ScreenReaderBase.cpp
--- a/Engine/Plugins/Experimental/ScreenReader/Source/ScreenReader/Private/GenericPlatform/ScreenReaderBase.cpp
+++ b/Engine/Plugins/Experimental/ScreenReader/Source/ScreenReader/Private/GenericPlatform/ScreenReaderBase.cpp
@@ -25,11 +25,15 @@ void FScreenReaderBase::Activate()
 {
 	if (!IsActive())
 	{
-		checkf(PlatformApplication.IsValid(), TEXT("Trying to activate the screen reader with an invalid platform aplication.with "));
+		checkf(PlatformApplication.IsValid(), TEXT("Trying to activate the screen reader with an invalid platform aplication."));
 		UE_LOG(LogScreenReader, Verbose, TEXT("Activating screen reader."));
 		// Set the screen reader application message handler to be the platform's message handler 
 		// This allows the screen reader to intercept all applications before passing them on to Slate 
 		PlatformApplication.Pin()->SetMessageHandler(ScreenReaderApplicationMessageHandler);
+		TSharedRef<FGenericAccessibleMessageHandler> AccessibleMessageHandler = PlatformApplication.Pin()->GetAccessibleMessageHandler();
+		// We unbind the delegate just in case it is still bound to some other target. This ensures that all accessible events now only get processed by the screen reader via OnAccessibleEventRaised
+		AccessibleMessageHandler->UnbindAccessibleEventDelegate();
+		AccessibleMessageHandler->SetAccessibleEventDelegate(FGenericAccessibleMessageHandler::FAccessibleEvent::CreateRaw(this, &FScreenReaderBase::OnAccessibleEventRaised));
 		OnActivate();
 		bActive = true;
 }
@@ -41,6 +45,8 @@ void FScreenReaderBase::Deactivate()
 	{
 		UE_LOG(LogScreenReader, Verbose, TEXT("Deactivating screen reader."));
 		OnDeactivate();
+		TSharedRef<FGenericAccessibleMessageHandler> AccessibleMessageHandler = PlatformApplication.Pin()->GetAccessibleMessageHandler();
+		AccessibleMessageHandler->UnbindAccessibleEventDelegate();
 		// Set FSlateApplication as the platform message ahndler again 
 		PlatformApplication.Pin()->SetMessageHandler(ScreenReaderApplicationMessageHandler->GetTargetMessageHandler());
 		bActive = false;
@@ -49,54 +55,40 @@ void FScreenReaderBase::Deactivate()
 
 bool FScreenReaderBase::RegisterUser(int32 InUserId)
 {
-	if (UsersMap.Contains(InUserId))
-	{
-		UE_LOG(LogScreenReader, Verbose, TEXT("Failed to register screen reader user with Id %d. Another user with the same Id has already been registered."), InUserId);
-		return false;
-	}
-	UsersMap.Add(InUserId, MakeShared<FScreenReaderUser>(InUserId));
-	UE_LOG(LogScreenReader, Verbose, TEXT("Registered screen reader user %d."), InUserId);
-	return true;
+	FGenericAccessibleUserRegistry& UserRegistry = PlatformApplication.Pin()->GetAccessibleMessageHandler()->GetAccessibleUserRegistry();
+	// @TODOAccessibility: Create a mechanism for creating the appropriate user. Factory or abstract factory something like that
+	TSharedRef<FScreenReaderUser> ScreenReaderUser = MakeShared<FScreenReaderUser>(InUserId);
+	return UserRegistry.RegisterUser(ScreenReaderUser);
 }
 
 bool FScreenReaderBase::UnregisterUser(int32 InUserId)
 {
-	if (!UsersMap.Contains(InUserId))
-	{
-		UE_LOG(LogScreenReader, Verbose, TEXT("Failed to unregister with Id %d. No user with %d Id is registered."), InUserId, InUserId);
-		return false;
-	}
-	UsersMap[InUserId]->Deactivate();
-	UsersMap.Remove(InUserId);
-	UE_LOG(LogScreenReader, Verbose, TEXT("Unregistered screen reader user %d."), InUserId);
-	return true;
+	FGenericAccessibleUserRegistry& UserRegistry = PlatformApplication.Pin()->GetAccessibleMessageHandler()->GetAccessibleUserRegistry();
+	return UserRegistry.UnregisterUser(InUserId);
 }
 
 bool FScreenReaderBase::IsUserRegistered(int32 InUserId) const
 {
-	return UsersMap.Contains(InUserId);
+	const FGenericAccessibleUserRegistry& UserRegistry = PlatformApplication.Pin()->GetAccessibleMessageHandler()->GetAccessibleUserRegistry();
+	return UserRegistry.IsUserRegistered(InUserId);
 }
 
 void FScreenReaderBase::UnregisterAllUsers()
 {
-	for (TPair<int32, TSharedRef<FScreenReaderUser>>& UserPair : UsersMap)
-	{
-		UserPair.Value->Deactivate();
-	}
-	UsersMap.Empty();
+	FGenericAccessibleUserRegistry& UserRegistry = PlatformApplication.Pin()->GetAccessibleMessageHandler()->GetAccessibleUserRegistry();
+	UserRegistry.UnregisterAllUsers();
 }
 
 TSharedRef<FScreenReaderUser> FScreenReaderBase::GetUserChecked(int32 InUserId) const
 {
-	checkf(IsUserRegistered(InUserId), TEXT("User Id %d is not registered. Did you forget to register the user with RegisterUser()?"), InUserId);
-	return UsersMap[InUserId];
+	const FGenericAccessibleUserRegistry& UserRegistry = PlatformApplication.Pin()->GetAccessibleMessageHandler()->GetAccessibleUserRegistry();
+	checkf(UserRegistry.IsUserRegistered(InUserId), TEXT("User Id %d is not registered. Did you forget to register the user with RegisterUser()?"), InUserId);
+	return StaticCastSharedRef<FScreenReaderUser>(UserRegistry.GetUser(InUserId).ToSharedRef());
 }
 	
 TSharedPtr<FScreenReaderUser> FScreenReaderBase::GetUser(int32 InUserId) const
 {
-	if (IsUserRegistered(InUserId))
-	{
-		return UsersMap[InUserId];
-	}
-	return nullptr;
-}
\ No newline at end of file
+	const FGenericAccessibleUserRegistry& UserRegistry = PlatformApplication.Pin()->GetAccessibleMessageHandler()->GetAccessibleUserRegistry();
+	return StaticCastSharedPtr<FScreenReaderUser>(UserRegistry.GetUser(InUserId));
+}
+
diff --git a/Engine/Plugins/Experimental/ScreenReader/Source/ScreenReader/Private/GenericPlatform/ScreenReaderUser.cpp b/Engine/Plugins/Experimental/ScreenReader/Source/ScreenReader/Private/GenericPlatform/ScreenReaderUser.cpp
--- a/Engine/Plugins/Experimental/ScreenReader/Source/ScreenReader/Private/GenericPlatform/ScreenReaderUser.cpp
+++ b/Engine/Plugins/Experimental/ScreenReader/Source/ScreenReader/Private/GenericPlatform/ScreenReaderUser.cpp
@@ -4,8 +4,8 @@
 #include "Announcement/ScreenReaderAnnouncementChannel.h"
 #include "TextToSpeech.h"
 
-FScreenReaderUser::FScreenReaderUser(int32 InUserId)
-	: UserId(InUserId)
+FScreenReaderUser::FScreenReaderUser(FAccessibleUserIndex InUserIndex)
+	: FGenericAccessibleUser(InUserIndex)
 	, bActive(false)
 {
 	// @TODOAccessibility: For now, just default it to the platform TTS. Should give a way to allow users to use custom TTS though 
@@ -113,4 +113,10 @@ FScreenReaderReply FScreenReaderUser::UnmuteSpeech()
 bool FScreenReaderUser::IsSpeechMuted() const
 {
 	return AnnouncementChannel->IsSpeechMuted();
-}
\ No newline at end of file
+}
+
+void FScreenReaderUser::OnUnregistered()
+{
+	Deactivate();
+}
+
diff --git a/Engine/Plugins/Experimental/ScreenReader/Source/ScreenReader/Public/GenericPlatform/ScreenReaderBase.h b/Engine/Plugins/Experimental/ScreenReader/Source/ScreenReader/Public/GenericPlatform/ScreenReaderBase.h
--- a/Engine/Plugins/Experimental/ScreenReader/Source/ScreenReader/Public/GenericPlatform/ScreenReaderBase.h
+++ b/Engine/Plugins/Experimental/ScreenReader/Source/ScreenReader/Public/GenericPlatform/ScreenReaderBase.h
@@ -10,13 +10,12 @@ class GenericApplication;
 
 
 /**
-* The base class all screen readers must derive from.
+* The abstract base class all screen readers must derive from.
 * The screen reader sets up and tears down the screen reader framework with Activate() and Deactivate() respectively.
 * When active, the screen reader will intercept all application messages, perform necessary input pre-processing and dispatch the accessible events and 
 * inputs to all registered screen reader users.
-* The screen reader contains a registry of screen reader users which can be registered, unregistered and retrieved with corresponding user Ids.
-* Registered users will receive and respond to accessible inputs and accessible events dispatched by the screen reader.
-* Users of this class and its child classes must activate the screen reader to ensure the accessible inputs and accessible events are dispatched properly.
+* Users of derived classes must activate the screen reader to ensure the accessible inputs and accessible events are dispatched properly.
+ * Derived classes must implement the pure virtual OnAccessibleEventRaised function to handle accessible events.
 * @see FScreenReaderUser
 */
 class SCREENREADER_API FScreenReaderBase
@@ -92,16 +91,29 @@ protected:
 	* feedback to alert the user that the screen reader is deactivated. 
 	*/
 	virtual void OnDeactivate() {}
+	/**
+	 * The handler for all accessible events raised from FGenericAccessibleMessageHandler. Is bound to the accessible event delegate in FGenericAccessibleMessageHandler.
+	 * This handler is bound during activation and unbound during deactivation.
+	 * A sample implementation can be found in FSlateScreenReader.
+	 * @see FSlateScreenReader
+	 */
+virtual void OnAccessibleEventRaised(const FAccessibleEventArgs& Args) = 0;
+	
+	/** Returns the weak pointer to the platform application. Use this to access the accessible message handler and or retrieve cursor information. */
+	TWeakPtr<GenericApplication> GetPlatformApplicationWeak() const
+	{
+		return PlatformApplication;
+	}
+
 private:
 	/** 
 	* The screen reader's application message handler. It processes input events before passing them down to the target application message handler
 	* @see FScreenReaderApplicationMessageHandlerBase
 	*/
 	TSharedRef<FScreenReaderApplicationMessageHandlerBase> ScreenReaderApplicationMessageHandler;
-	/** The platform application. Used to access the accessible message handler and swap the application message handler with the screen reader application message handler.. */
+	/** The platform application. Used to access the accessible message handler and swap the application message handler with the screen reader application message handler. */
 	TWeakPtr<GenericApplication> PlatformApplication;
-	/** A map of user Ids to screen reader users. Each entry represents a registered user that can receive and repsond to accessible events and accessible input. */
-	TMap<int32, TSharedRef<FScreenReaderUser>> UsersMap;
 	/** True if the screen reader is active. Else false. */
 	bool bActive;
-};
\ No newline at end of file
+};
+
diff --git a/Engine/Plugins/Experimental/ScreenReader/Source/ScreenReader/Public/GenericPlatform/ScreenReaderUser.h b/Engine/Plugins/Experimental/ScreenReader/Source/ScreenReader/Public/GenericPlatform/ScreenReaderUser.h
--- a/Engine/Plugins/Experimental/ScreenReader/Source/ScreenReader/Public/GenericPlatform/ScreenReaderUser.h
+++ b/Engine/Plugins/Experimental/ScreenReader/Source/ScreenReader/Public/GenericPlatform/ScreenReaderUser.h
@@ -4,6 +4,7 @@
 #include "CoreMinimal.h"
 #include "Announcement/ScreenReaderAnnouncement.h"
 #include "GenericPlatform/ScreenReaderReply.h"
+#include "GenericPlatform/Accessibility/GenericAccessibleInterfaces.h"
 
 class IAccessibleWidget;
 class FScreenReaderAnnouncementChannel;
@@ -13,16 +14,16 @@ class FScreenReaderAnnouncementChannel;
 * This class is a facade that acts as a one stop shop for all screen reading services on a per user basis.
 * Multiple screen reader users can exist simultaneoulsy to facilitate local multiplayer and users can opt in or out 
 * of receiving screen reader feedback by registering and unregistering with the screen reader respectively.
-* A screen user is inactive by default when it is first registered with teh screen reader. Users must explicitly activate the screen reader user after successful registration with the screen reader to use its services.
+* A screen user is inactive by default when it is first registered through the screen reader. Users must explicitly activate the screen reader user after successful registration with the screen reader to use its services.
 * Responsibilities of the class are:
 * 1. Text to speech (TTS) requests - Users can request an announcement to be spoken via text to speech to  screen reader user
 * 2. Accessible focus handling - A screen reader user holds information about the accessible widget it is currently focused on
-* @see FScreenReaderBase, FScreenReaderAnnouncement
+* @see FScreenReaderBase, FScreenReaderAnnouncement, FGenericAccessibleUser
 */
-class SCREENREADER_API FScreenReaderUser
+class SCREENREADER_API FScreenReaderUser : public FGenericAccessibleUser
 {
 public:
-	explicit FScreenReaderUser(int32 InUserId);
+	explicit FScreenReaderUser(FAccessibleUserIndex InUserIndex);
 	~FScreenReaderUser();
 	/**
 	* Requests an announcement to be spoken to the screen reader user. This is the main mechanism to provide text to speech auditory feedback
@@ -33,7 +34,7 @@ public:
 	* If the screen reader user is active and no announcements are currently spoken, the announcement will be spoken immediately.
 	* If another announcement is currently being spoken, the passed in announcement could be queued or interrupt the currently spoken announcement.
 	* Examples:
-	* TSharedRef<FScreenReaderUser> MyUser = ScreenReader->GetUser(MyScreenReaderUserId);
+	* TSharedRef<FScreenReaderUser> MyUser = ScreenReader->GetUser(MyScreenReaderUserIndex);
 	* // All announcements to be spoken should be localized to provide language support for the text to speech system.
 	* static const FText MyText = LOCTEXT("ExampleUserFeedback", "Feedback to user.");
 	* // This makes the announcement interrupt any currently spoken announcement with lower priority and is played immediately. The announcement will be uninterruptable and guarantees the user will hear the announcement.
@@ -104,25 +105,15 @@ public:
 	{ 
 		return bActive; 
 	}
-	/** Returns the user Id associated with this screen reader user. */
-	int32 GetUserId() const 
-	{ 
-		return UserId; 
-	}
-	/** Returns the accessible widget the screen reader user is currently focused on. */
-	TSharedPtr<IAccessibleWidget> GetAccessibleFocusWidget() const 
-	{ 
-		return AccessibleFocusWidget.Pin();
-	}
-	/** Sets the accessible widget the screen reader user is currently focused on */
-	void SetAccessibleFocusWidget(const TSharedRef<IAccessibleWidget>& InAccessibleFocusWidget) 
-	{ 
-		AccessibleFocusWidget = InAccessibleFocusWidget;
-	}
+
+protected:
+	//~ Begin FGenericAccessibleUser interface
+	virtual void OnUnregistered() override;
+	//~ End FGenericAccessibleUser interface
+	
 private:
-	int32 UserId;
 	/** Responsible for handling all incoming announcement requests and speaking them via text to speech if possible */
 	TUniquePtr<FScreenReaderAnnouncementChannel> AnnouncementChannel;
-	TWeakPtr<IAccessibleWidget> AccessibleFocusWidget;
 	bool bActive;
 };
+
diff --git a/Engine/Plugins/Experimental/SlateScreenReader/Source/SlateScreenReader/Private/SlateScreenReader.cpp b/Engine/Plugins/Experimental/SlateScreenReader/Source/SlateScreenReader/Private/SlateScreenReader.cpp
--- a/Engine/Plugins/Experimental/SlateScreenReader/Source/SlateScreenReader/Private/SlateScreenReader.cpp
+++ b/Engine/Plugins/Experimental/SlateScreenReader/Source/SlateScreenReader/Private/SlateScreenReader.cpp
@@ -5,21 +5,16 @@
 #include "GenericPlatform/Accessibility/GenericAccessibleInterfaces.h"
 #include "Announcement/ScreenReaderAnnouncementChannel.h"
 #include "GenericPlatform/ScreenReaderUser.h"
-#include "InputCoreTypes.h"
 #include "SlateScreenReaderLog.h"
 
-// Slate includes
-#include "Framework/Application/SlateApplication.h"
-#include "Framework/Application/SlateUser.h"
-#include "Input/Events.h"
-#include "Widgets/Accessibility/SlateAccessibleWidgetCache.h"
-
-#define LOCTEXT_NAMESPACE "SlateScreenReader"
 
 FSlateScreenReader::FSlateScreenReader(const TSharedRef<GenericApplication>& InPlatformApplication)
 	: FScreenReaderBase(InPlatformApplication)
 {
-
+	// if we are creating and using FSlateScreenReader, we assume that other accessibility options (such as platform accessibility) are going to be overridden
+	// We unregister all previously registered users and force registration through the screen reader framework
+	// @TODOAccessibility: It would be better to swap the existing FGenericAccessibleUsers for their FScreenReaderUser counterparts, but we will deal with that later
+	InPlatformApplication->GetAccessibleMessageHandler()->GetAccessibleUserRegistry().UnregisterAllUsers();
 }
 
 FSlateScreenReader::~FSlateScreenReader()
@@ -27,32 +22,40 @@ FSlateScreenReader::~FSlateScreenReader()
 
 }
 
-void FSlateScreenReader::HandleSlateFocusChanging(const FFocusEvent& InFocusEvent, const FWeakWidgetPath& InOldFocusWidgetPath, const TSharedPtr<SWidget>& InOldWidget, const FWidgetPath& InNewWidgetPath, const TSharedPtr<SWidget>& InNewFocusWidget)
+void FSlateScreenReader::OnAccessibleEventRaised(const FAccessibleEventArgs& Args)
+{
+	// This should only be triggered by the accessible message handler which initiates from the Slate thread.
+	check(IsInGameThread());
+if (IsActive())
 {
-	if (IsActive() && InNewFocusWidget && InNewFocusWidget->IsAccessible())
+	switch (Args.Event)
 	{
-		int32 UserId = static_cast<int32>(InFocusEvent.GetUser());
-		if (TSharedPtr<FScreenReaderUser> User = GetUser(UserId))
+		case EAccessibleEvent::FocusChange:
 		{
-			TSharedPtr<IAccessibleWidget> AccessibleWidget = FSlateAccessibleWidgetCache::GetAccessibleWidgetChecked(InNewFocusWidget);
-			if (AccessibleWidget)
+			// if the widget is gaining focus
+			if (Args.NewValue.GetValue<bool>())
 			{
-				User->SetAccessibleFocusWidget(AccessibleWidget.ToSharedRef());
-				User->RequestSpeakWidget(AccessibleWidget.ToSharedRef());
+				TSharedPtr<FScreenReaderUser> User = GetUser(Args.UserIndex);
+				if (User)
+				{
+					User->RequestSpeakWidget(Args.Widget);
+				}
 			}
+			break;
 		}
 	}
 }
+}
 
 void FSlateScreenReader::OnActivate()
 {
 	UE_LOG(LogSlateScreenReader, Verbose, TEXT("On activation of Slate screen reader."));
-	// @TODOAccessibility: Consider removing the binding on deactivation 
-	FSlateApplication::Get().OnFocusChanging().AddRaw(this, &FSlateScreenReader::HandleSlateFocusChanging);
+	// @TODOAccessibility: Provide feedback in the form of a tone or something
 }
 
 void FSlateScreenReader::OnDeactivate()
 {
 	UE_LOG(LogSlateScreenReader, Verbose, TEXT("Deactivation of Slate screen reader."));
+	// @TODOAccessibility: Provide feedback in the form of a tone or something
 } 
-#undef LOCTEXT_NAMESPACE
+
diff --git a/Engine/Plugins/Experimental/SlateScreenReader/Source/SlateScreenReader/Private/SlateScreenReader.h b/Engine/Plugins/Experimental/SlateScreenReader/Source/SlateScreenReader/Private/SlateScreenReader.h
--- a/Engine/Plugins/Experimental/SlateScreenReader/Source/SlateScreenReader/Private/SlateScreenReader.h
+++ b/Engine/Plugins/Experimental/SlateScreenReader/Source/SlateScreenReader/Private/SlateScreenReader.h
@@ -8,9 +8,7 @@
 
 class SWidget;
 class iAccessibleWidget;
-class FWidgetPath;
-class FWeakWidgetPath;
-struct FFocusEvent;
+
 
 /**
  * A basic screen reader that can work for desktop and consoles that use Slate.
@@ -23,12 +21,15 @@ public:
 	FSlateScreenReader() = delete;
 	FSlateScreenReader(const FSlateScreenReader& Other) = delete;
 	explicit FSlateScreenReader(const TSharedRef<GenericApplication>& InPlatformApplication);
-
-	FSlateScreenReader& operator=(const FSlateScreenReader& Other) = delete;
 	virtual ~FSlateScreenReader();
 	
-	void HandleSlateFocusChanging(const FFocusEvent& InFocusEvent, const FWeakWidgetPath& InOldFocusWidgetPath, const TSharedPtr<SWidget>& InOldWidget, const FWidgetPath& InNewWidgetPath, const TSharedPtr<SWidget>& InNewFocusWidget);
+	FSlateScreenReader& operator=(const FSlateScreenReader& Other) = delete;
+	
 protected:
+	//~ Begin FScreenReaderBase interface
 	virtual void OnActivate() override;
 	virtual void OnDeactivate() override;
+	virtual void OnAccessibleEventRaised(const FAccessibleEventArgs& Args) override;
+	//~ End FScreenReaderBase interface
 };
+
diff --git a/Engine/Plugins/Experimental/SlateScreenReader/Source/SlateScreenReader/Private/SlateScreenReaderEngineSubsystem.cpp b/Engine/Plugins/Experimental/SlateScreenReader/Source/SlateScreenReader/Private/SlateScreenReaderEngineSubsystem.cpp
--- a/Engine/Plugins/Experimental/SlateScreenReader/Source/SlateScreenReader/Private/SlateScreenReaderEngineSubsystem.cpp
+++ b/Engine/Plugins/Experimental/SlateScreenReader/Source/SlateScreenReader/Private/SlateScreenReaderEngineSubsystem.cpp
@@ -159,7 +159,7 @@ FScreenReaderReply USlateScreenReaderEngineSubsystem::RequestSpeakFocusedWidget(
 		TSharedPtr<FScreenReaderUser> User = ScreenReader->GetUser(InUserId);
 		if (User)
 		{
-			if (TSharedPtr<IAccessibleWidget> UserFocusedWidget = User->GetAccessibleFocusWidget())
+			if (TSharedPtr<IAccessibleWidget> UserFocusedWidget = User->GetFocusedAccessibleWidget())
 			{
 				return User->RequestSpeakWidget(UserFocusedWidget.ToSharedRef());
 			}
