Speculative fix for compilation errors in DevBuild relating to screen reader..
#rnx

diff --git a/Engine/Plugins/Experimental/ScreenReader/Source/ScreenReader/Public/GenericPlatform/ScreenReaderBase.h b/Engine/Plugins/Experimental/ScreenReader/Source/ScreenReader/Public/GenericPlatform/ScreenReaderBase.h
--- a/Engine/Plugins/Experimental/ScreenReader/Source/ScreenReader/Public/GenericPlatform/ScreenReaderBase.h
+++ b/Engine/Plugins/Experimental/ScreenReader/Source/ScreenReader/Public/GenericPlatform/ScreenReaderBase.h
@@ -7,6 +7,7 @@ class FGenericApplicationMessageHandler;
 class FScreenReaderApplicationMessageHandlerBase;
 class FScreenReaderUser;
 class GenericApplication;
+struct FAccessibleEventArgs;
 
 
 /**
@@ -97,8 +98,7 @@ protected:
 	 * A sample implementation can be found in FSlateScreenReader.
 	 * @see FSlateScreenReader
 	 */
-virtual void OnAccessibleEventRaised(const FAccessibleEventArgs& Args) = 0;
-	
+	virtual void OnAccessibleEventRaised(const FAccessibleEventArgs& Args) = 0;
 	/** Returns the weak pointer to the platform application. Use this to access the accessible message handler and or retrieve cursor information. */
 	TWeakPtr<GenericApplication> GetPlatformApplicationWeak() const
 	{
diff --git a/Engine/Plugins/Experimental/SlateScreenReader/Source/SlateScreenReader/Private/SlateScreenReader.h b/Engine/Plugins/Experimental/SlateScreenReader/Source/SlateScreenReader/Private/SlateScreenReader.h
--- a/Engine/Plugins/Experimental/SlateScreenReader/Source/SlateScreenReader/Private/SlateScreenReader.h
+++ b/Engine/Plugins/Experimental/SlateScreenReader/Source/SlateScreenReader/Private/SlateScreenReader.h
@@ -6,8 +6,8 @@
 #include "GenericPlatform/ScreenReaderBase.h"
 #include "GenericPlatform/GenericApplicationMessageHandler.h"
 
-class SWidget;
 class iAccessibleWidget;
+struct FAccessibleEventArgs;
 
 
 /**
