Update the ScreenReaderApplication to use the newer version of OnController functions instead of the soon to be deprecated version

#rnx

diff --git a/Engine/Plugins/Experimental/ScreenReader/Source/ScreenReader/Private/GenericPlatform/ScreenReaderApplicationMessageHandlerBase.cpp b/Engine/Plugins/Experimental/ScreenReader/Source/ScreenReader/Private/GenericPlatform/ScreenReaderApplicationMessageHandlerBase.cpp
--- a/Engine/Plugins/Experimental/ScreenReader/Source/ScreenReader/Private/GenericPlatform/ScreenReaderApplicationMessageHandlerBase.cpp
+++ b/Engine/Plugins/Experimental/ScreenReader/Source/ScreenReader/Private/GenericPlatform/ScreenReaderApplicationMessageHandlerBase.cpp
@@ -88,6 +88,7 @@ bool FScreenReaderApplicationMessageHandlerBase::OnCursorSet()
 	return TargetMessageHandler->OnCursorSet();
 }
 
+PRAGMA_DISABLE_DEPRECATION_WARNINGS
 bool FScreenReaderApplicationMessageHandlerBase::OnControllerAnalog(FGamepadKeyNames::Type KeyName, int32 ControllerId, float AnalogValue)
 {
 	return TargetMessageHandler->OnControllerAnalog(KeyName, ControllerId, AnalogValue);
@@ -102,6 +103,22 @@ bool FScreenReaderApplicationMessageHandlerBase::OnControllerButtonReleased(FGam
 {
 	return TargetMessageHandler->OnControllerButtonReleased(KeyName, ControllerId, IsRepeat);
 }
+PRAGMA_ENABLE_DEPRECATION_WARNINGS
+
+bool FScreenReaderApplicationMessageHandlerBase::OnControllerAnalog(FGamepadKeyNames::Type KeyName, FPlatformUserId PlatformUserId, FInputDeviceId InputDeviceId, float AnalogValue)
+{
+	return TargetMessageHandler->OnControllerAnalog(KeyName, PlatformUserId, InputDeviceId, AnalogValue);
+}
+
+bool FScreenReaderApplicationMessageHandlerBase::OnControllerButtonPressed(FGamepadKeyNames::Type KeyName, FPlatformUserId PlatformUserId, FInputDeviceId InputDeviceId, bool IsRepeat)
+{
+	return TargetMessageHandler->OnControllerButtonPressed(KeyName, PlatformUserId, InputDeviceId, IsRepeat);
+}
+
+bool FScreenReaderApplicationMessageHandlerBase::OnControllerButtonReleased(FGamepadKeyNames::Type KeyName, FPlatformUserId PlatformUserId, FInputDeviceId InputDeviceId, bool IsRepeat)
+{
+	return TargetMessageHandler->OnControllerButtonReleased(KeyName, PlatformUserId, InputDeviceId, IsRepeat);
+}
 
 void FScreenReaderApplicationMessageHandlerBase::OnBeginGesture()
 {
diff --git a/Engine/Plugins/Experimental/ScreenReader/Source/ScreenReader/Public/GenericPlatform/ScreenReaderApplicationMessageHandlerBase.h b/Engine/Plugins/Experimental/ScreenReader/Source/ScreenReader/Public/GenericPlatform/ScreenReaderApplicationMessageHandlerBase.h
--- a/Engine/Plugins/Experimental/ScreenReader/Source/ScreenReader/Public/GenericPlatform/ScreenReaderApplicationMessageHandlerBase.h
+++ b/Engine/Plugins/Experimental/ScreenReader/Source/ScreenReader/Public/GenericPlatform/ScreenReaderApplicationMessageHandlerBase.h
@@ -53,9 +53,18 @@ public:
 	virtual bool OnMouseMove() override;
 	virtual bool OnRawMouseMove(const int32 X, const int32 Y) override;
 	virtual bool OnCursorSet() override;
+
+	UE_DEPRECATED(5.1, "This version of OnControllerAnalog is deprecated, please use the version that takes an FPlatformUserId instead")
 	virtual bool OnControllerAnalog(FGamepadKeyNames::Type KeyName, int32 ControllerId, float AnalogValue) override;
+	UE_DEPRECATED(5.1, "This version of OnControllerButtonPressed is deprecated, please use the version that takes an FPlatformUserId instead")
 	virtual bool OnControllerButtonPressed(FGamepadKeyNames::Type KeyName, int32 ControllerId, bool IsRepeat) override;
+	UE_DEPRECATED(5.1, "This version of OnControllerButtonReleased is deprecated, please use the version that takes an FPlatformUserId instead")
 	virtual bool OnControllerButtonReleased(FGamepadKeyNames::Type KeyName, int32 ControllerId, bool IsRepeat) override;
+
+	virtual bool OnControllerAnalog(FGamepadKeyNames::Type KeyName, FPlatformUserId PlatformUserId, FInputDeviceId InputDeviceId, float AnalogValue) override;
+	virtual bool OnControllerButtonPressed(FGamepadKeyNames::Type KeyName, FPlatformUserId PlatformUserId, FInputDeviceId InputDeviceId, bool IsRepeat) override;
+	virtual bool OnControllerButtonReleased(FGamepadKeyNames::Type KeyName, FPlatformUserId PlatformUserId, FInputDeviceId InputDeviceId, bool IsRepeat) override;
+	
 	virtual void OnBeginGesture() override;
 	virtual bool OnTouchGesture(EGestureEvent GestureType, const FVector2D& Delta, float WheelDelta, bool bIsDirectionInvertedFromDevice) override;
 	virtual void OnEndGesture() override;
