Fix stack use after scope issue in screen reader.

#tests Enabled the screen reader plugin in engine test, started the editor and verified that everything went well when executing the modified code.
[RN] Fixed an Asan issue with Slate Screen Reader at initialization of the plugin.

diff --git a/Engine/Plugins/Experimental/ScreenReader/Source/ScreenReader/Public/GenericPlatform/IScreenReaderBuilder.h b/Engine/Plugins/Experimental/ScreenReader/Source/ScreenReader/Public/GenericPlatform/IScreenReaderBuilder.h
--- a/Engine/Plugins/Experimental/ScreenReader/Source/ScreenReader/Public/GenericPlatform/IScreenReaderBuilder.h
+++ b/Engine/Plugins/Experimental/ScreenReader/Source/ScreenReader/Public/GenericPlatform/IScreenReaderBuilder.h
@@ -27,7 +27,7 @@ public:
 		{}
 		~FArgs() = default;
 		/** The platform application to be passed to the screen reader for initialization. */
-		const TSharedRef<GenericApplication>& PlatformApplication;
+		const TSharedRef<GenericApplication> PlatformApplication;
 	};
 	/**
 	* A pure virtual function to build a screen reader.
